#ue4
Start of 'optimized' Blueprint compiles:
- Disabled debug and trace site generation when compiling in a commandlet (such as during cooking)

diff --git a/Engine/Source/Editor/BlueprintGraph/Private/K2Node_ExecutionSequence.cpp b/Engine/Source/Editor/BlueprintGraph/Private/K2Node_ExecutionSequence.cpp
--- a/Engine/Source/Editor/BlueprintGraph/Private/K2Node_ExecutionSequence.cpp
+++ b/Engine/Source/Editor/BlueprintGraph/Private/K2Node_ExecutionSequence.cpp
@@ -48,8 +48,7 @@ public:
 		// Process the pins, if there are any valid entries
 		if (OutputPins.Num() > 0)
 		{
-			const bool bGenerateDebugSites = true;
-			if ((KismetCompilerDebugOptions::EmitDebuggingPlaceholders) && (OutputPins.Num() > 1))
+			if (Context.bCreateDebugData && (OutputPins.Num() > 1))
 			{
 				const FString NodeComment = Node->NodeComment.IsEmpty() ? Node->GetName() : Node->NodeComment;
 
diff --git a/Engine/Source/Editor/KismetCompiler/Private/KismetCompiler.cpp b/Engine/Source/Editor/KismetCompiler/Private/KismetCompiler.cpp
--- a/Engine/Source/Editor/KismetCompiler/Private/KismetCompiler.cpp
+++ b/Engine/Source/Editor/KismetCompiler/Private/KismetCompiler.cpp
@@ -1206,9 +1206,9 @@ void FKismetCompilerContext::CompileFunction(FKismetFunctionContext& Context)
 		}
 
 		// Debug opcode insertion point
-		if (KismetCompilerDebugOptions::EmitDebuggingPlaceholders && !IsNodePure(Node))
+		if (Context.bCreateDebugData && !IsNodePure(Node))
 		{
-			bool bEmitDebuggingSite = Context.bCreateDebugData;
+			bool bEmitDebuggingSite = true;
 
 			if (Context.IsEventGraph() && (Node->IsA(UK2Node_FunctionEntry::StaticClass())))
 			{
diff --git a/Engine/Source/Editor/KismetCompiler/Private/KismetCompilerMisc.cpp b/Engine/Source/Editor/KismetCompiler/Private/KismetCompilerMisc.cpp
--- a/Engine/Source/Editor/KismetCompiler/Private/KismetCompilerMisc.cpp
+++ b/Engine/Source/Editor/KismetCompiler/Private/KismetCompilerMisc.cpp
@@ -860,6 +860,13 @@ FKismetFunctionContext::FKismetFunctionContext(FCompilerResultsLog& InMessageLog
 {
 	NetNameMap = new FNetNameMapping();
 	bAllocatedNetNameMap = true;
+
+	// Prevent debug generation when cooking or running other commandlets
+	// Compile-on-load will recreate it if the editor is run
+	if (IsRunningCommandlet())
+	{
+		bCreateDebugData = false;
+	}
 }
 
 FKismetFunctionContext::~FKismetFunctionContext()
diff --git a/Engine/Source/Editor/KismetCompiler/Public/KismetCompiledFunctionContext.h b/Engine/Source/Editor/KismetCompiler/Public/KismetCompiledFunctionContext.h
--- a/Engine/Source/Editor/KismetCompiler/Public/KismetCompiledFunctionContext.h
+++ b/Engine/Source/Editor/KismetCompiler/Public/KismetCompiledFunctionContext.h
@@ -17,9 +17,6 @@ namespace KismetCompilerDebugOptions
 
 	// Should the compiler emit node comments to the backends?
 	enum { EmitNodeComments = DebuggingCompiler };
-
-	// Should the compiler emit a NOP in each execution state as a placeholder for breakpoints and trace instructions?
-	enum { EmitDebuggingPlaceholders = 1 };
 }
 
 //////////////////////////////////////////////////////////////////////////
