Added include directives to fix non-Unity build

diff --git a/Engine/Source/Editor/BlueprintGraph/Private/K2Node_EnumLiteral.cpp b/Engine/Source/Editor/BlueprintGraph/Private/K2Node_EnumLiteral.cpp
--- a/Engine/Source/Editor/BlueprintGraph/Private/K2Node_EnumLiteral.cpp
+++ b/Engine/Source/Editor/BlueprintGraph/Private/K2Node_EnumLiteral.cpp
@@ -2,6 +2,7 @@
 
 
 #include "BlueprintGraphPrivatePCH.h"
+#include "KismetCompilerMisc.h"
 
 const FString& UK2Node_EnumLiteral::GetEnumInputPinName()
 {
diff --git a/Engine/Source/Editor/KismetCompiler/Public/KismetCompilerMisc.h b/Engine/Source/Editor/KismetCompiler/Public/KismetCompilerMisc.h
--- a/Engine/Source/Editor/KismetCompiler/Public/KismetCompilerMisc.h
+++ b/Engine/Source/Editor/KismetCompiler/Public/KismetCompilerMisc.h
@@ -5,6 +5,7 @@
 #include "EdGraphCompilerUtilities.h"
 #include "UnrealEd.h"
 #include "Editor.h"
+#include "KismetCompiledFunctionContext.h"
 
 //////////////////////////////////////////////////////////////////////////
 // FKismetCompilerUtilities
