TTP #321255: Adding a more detailed error when self pins don't match the current blueprint.

diff --git a/Engine/Source/Editor/KismetCompiler/Private/KismetCompiler.cpp b/Engine/Source/Editor/KismetCompiler/Private/KismetCompiler.cpp
--- a/Engine/Source/Editor/KismetCompiler/Private/KismetCompiler.cpp
+++ b/Engine/Source/Editor/KismetCompiler/Private/KismetCompiler.cpp
@@ -946,14 +946,23 @@ void FKismetCompilerContext::ValidateSelfPinsInGraph(const UEdGraph* SourceGraph
 					if (Schema->IsSelfPin(*Pin) && (Pin->LinkedTo.Num() == 0))
 					{
 						FEdGraphPinType SelfType;
-						SelfType.PinCategory = Schema->PC_Object;
+						SelfType.PinCategory    = Schema->PC_Object;
 						SelfType.PinSubCategory = Schema->PSC_Self;
 
 						if (!Schema->ArePinTypesCompatible(SelfType, Pin->PinType, NewClass))
 						{
 							if (Pin->DefaultObject == NULL)
 							{
-								MessageLog.Error(*LOCTEXT("PinMustHaveConnection_Error", "Pin @@ must have a connection").ToString(), Pin);
+								FString PinType = Pin->PinType.PinCategory;
+								if ((Pin->PinType.PinCategory == Schema->PC_Object)    || 
+									(Pin->PinType.PinCategory == Schema->PC_Interface) ||
+									(Pin->PinType.PinCategory == Schema->PC_Class))
+								{
+									PinType = Pin->PinType.PinSubCategoryObject->GetName();
+								}
+
+								FString ErrorMsg = FString::Printf(*LOCTEXT("PinMustHaveConnection_Error", "This blueprint (self) is not a %s, therefore '@@' must have a connection").ToString(), *PinType);
+								MessageLog.Error(*ErrorMsg, Pin);
 							}
 						}
 					}
