Fixing merged code that was missed with the switch over to MovePinLinksToIntermediate/CopyPinLinksToIntermediate for node expansion.

diff --git a/Engine/Source/Editor/KismetCompiler/Public/KismetCompiler.h b/Engine/Source/Editor/KismetCompiler/Public/KismetCompiler.h
--- a/Engine/Source/Editor/KismetCompiler/Public/KismetCompiler.h
+++ b/Engine/Source/Editor/KismetCompiler/Public/KismetCompiler.h
@@ -137,9 +137,6 @@ public:
 
 	KISMETCOMPILER_API UK2Node_TemporaryVariable* SpawnInternalVariable(UEdGraphNode* SourceNode, FString Category, FString SubCategory = TEXT(""), UObject* SubcategoryObject = NULL, bool bIsArray = false);
 
-	/** Checks a connection response, and errors if it didn't succeed */
-	KISMETCOMPILER_API void CheckConnectionResponse(const FPinConnectionResponse &Response, const UEdGraphNode *Node);
-
 protected:
 	virtual UEdGraphSchema_K2* CreateSchema();
 	virtual void PostCreateSchema();
@@ -158,6 +155,13 @@ protected:
 	 */
 	virtual void CleanAndSanitizeClass(UBlueprintGeneratedClass* ClassToClean, UObject*& OldCDO);
 
+	/** 
+	 * Checks a connection response, and errors if it didn't succeed (not public, 
+	 * users should be using MovePinLinksToIntermediate/CopyPinLinksToIntermediate 
+	 * instead of wrapping their own with this).
+	 */
+	KISMETCOMPILER_API void CheckConnectionResponse(const FPinConnectionResponse &Response, const UEdGraphNode *Node);
+
 protected:
 	// FGraphCompilerContext interface
 	virtual void ValidateLink(const UEdGraphPin* PinA, const UEdGraphPin* PinB) const OVERRIDE;
