Removing some debug code that was left in.

diff --git a/Engine/Source/Editor/KismetCompiler/Private/KismetCompiler.cpp b/Engine/Source/Editor/KismetCompiler/Private/KismetCompiler.cpp
--- a/Engine/Source/Editor/KismetCompiler/Private/KismetCompiler.cpp
+++ b/Engine/Source/Editor/KismetCompiler/Private/KismetCompiler.cpp
@@ -1330,12 +1330,6 @@ void FKismetCompilerContext::PostcompileFunction(FKismetFunctionContext& Context
 	// Sort the 'linear execution list' again by likely execution order.
 	Context.FinalSortLinearExecList();
 
-// 	if (Context.GetBlueprint()->GetLinker()->IsCooking())
-// 	{
-// 		int i = 0;
-// 		++i;
-// 	}
-
 	// Resolve goto links
 	Context.ResolveGotoFixups();
 
