#ue4 ttp331663 CRASH: SYNCED GAME: Client crashes when exiting to desktop from main menu or leaving the map

SkeletonClass should also be recompiled in ConditionalRecompileClass

diff --git a/Engine/Source/Editor/KismetCompiler/Private/KismetCompilerModule.cpp b/Engine/Source/Editor/KismetCompiler/Private/KismetCompilerModule.cpp
--- a/Engine/Source/Editor/KismetCompiler/Private/KismetCompilerModule.cpp
+++ b/Engine/Source/Editor/KismetCompiler/Private/KismetCompilerModule.cpp
@@ -117,7 +117,7 @@ void FKismet2CompilerModule::CompileBlueprint(class UBlueprint* Blueprint, const
 		FUserDefinedStructureCompilerUtils::CompileStructs(Blueprint, Results, false);
 	}
 
-	if((CompileOptions.CompileType != EKismetCompileType::BytecodeOnly) && (CompileOptions.CompileType != EKismetCompileType::StructuresOnly))
+	if (CompileOptions.CompileType != EKismetCompileType::StructuresOnly)
 	{
 		FBlueprintCompileReinstancer SkeletonReinstancer(Blueprint->SkeletonGeneratedClass);
 
diff --git a/Engine/Source/Runtime/Engine/Private/BlueprintGeneratedClass.cpp b/Engine/Source/Runtime/Engine/Private/BlueprintGeneratedClass.cpp
--- a/Engine/Source/Runtime/Engine/Private/BlueprintGeneratedClass.cpp
+++ b/Engine/Source/Runtime/Engine/Private/BlueprintGeneratedClass.cpp
@@ -239,6 +239,10 @@ void UBlueprintGeneratedClass::ConditionalRecompileClass(TArray<UObject*>* ObjLo
 		else if (FConditionalRecompileClassHepler::StaticLink == NecessaryAction)
 		{
 			StaticLink(true);
+			if (*GeneratingBP->SkeletonGeneratedClass)
+			{
+				GeneratingBP->SkeletonGeneratedClass->StaticLink(true);
+			}
 		}
 	}
 }
