#summary Reparenting a Blueprint with a timeline having the same name as a newly inherited variable will correctly rename the Timeline

#change When auto-renaming the Timeline, the visible Timeline name is used instead of the hidden name (suffixed with _Template).

#ttp 327552 - ENSURE: Ensure condition failed: Timeline->GetOuter() Function: UBlueprint::FindTimelineTemplateByVariableName() FKismetCompilerContext::ValidateTimelineNames()

diff --git a/Engine/Source/Editor/KismetCompiler/Private/KismetCompiler.cpp b/Engine/Source/Editor/KismetCompiler/Private/KismetCompiler.cpp
--- a/Engine/Source/Editor/KismetCompiler/Private/KismetCompiler.cpp
+++ b/Engine/Source/Editor/KismetCompiler/Private/KismetCompiler.cpp
@@ -425,9 +425,12 @@ void FKismetCompilerContext::ValidateTimelineNames()
 		{
 			if( ParentBPNameValidator.IsValid() && ParentBPNameValidator->IsValid(TimelineTemplate->GetName()) != EValidatorResult::Ok )
 			{
-				FName NewName = FBlueprintEditorUtils::FindUniqueKismetName(Blueprint, TimelineTemplate->GetName());
+				// Use the viewer displayed Timeline name (without the _Template suffix) because it will be added later for appropriate checks.
+				FString TimelineName = UTimelineTemplate::TimelineTemplateNameToVariableName(TimelineTemplate->GetFName());
+
+				FName NewName = FBlueprintEditorUtils::FindUniqueKismetName(Blueprint, TimelineName);
 				MessageLog.Warning(*FString::Printf(*LOCTEXT("TimelineConflictWarning", "Found a timeline with a conflicting name (%s) - changed to %s.").ToString(), *TimelineTemplate->GetName(), *NewName.ToString()));
-				FBlueprintEditorUtils::RenameTimeline(Blueprint, TimelineTemplate->GetFName(), NewName);
+				FBlueprintEditorUtils::RenameTimeline(Blueprint, FName(*TimelineName), NewName);
 			}
 		}
 	}
