#ue4 BP: Default value from an unconnected structure pin, wont cause "Missing opening parenthesis" message.  All pins/nodes set an empty string as a defeault value for structures, but UStructProperty::ImportText_Internal requires at least parenthesis "()".

diff --git a/Engine/Source/Editor/KismetCompiler/Private/KismetCompilerVMBackend.cpp b/Engine/Source/Editor/KismetCompiler/Private/KismetCompilerVMBackend.cpp
--- a/Engine/Source/Editor/KismetCompiler/Private/KismetCompilerVMBackend.cpp
+++ b/Engine/Source/Editor/KismetCompiler/Private/KismetCompilerVMBackend.cpp
@@ -458,7 +458,7 @@ public:
 					}
 
 					// Assume that any errors on the import of the name string have been caught in the function call generation
-					StructProperty->ImportText(*Term->Name, StructData, 0, NULL, GLog);
+					StructProperty->ImportText(Term->Name.IsEmpty() ? TEXT("()") : *Term->Name, StructData, 0, NULL, GLog);
 
  					Writer << EX_StructConst;
 					Writer << Struct;
