#ue4 ttp#332960 BLUEPRINTS: Latent action in child BP foiled/broken by parent BP
The problem was, that both delayed actions have the same UUID.

FKismetFunctionContext::GetContextUniqueID() returns now a 'more unique' value.

diff --git a/Engine/Source/Editor/KismetCompiler/Private/KismetCompilerMisc.cpp b/Engine/Source/Editor/KismetCompiler/Private/KismetCompilerMisc.cpp
--- a/Engine/Source/Editor/KismetCompiler/Private/KismetCompilerMisc.cpp
+++ b/Engine/Source/Editor/KismetCompiler/Private/KismetCompilerMisc.cpp
@@ -848,7 +848,6 @@ FKismetFunctionContext::FKismetFunctionContext(FCompilerResultsLog& InMessageLog
 	, NewClass(InNewClass)
 	, MessageLog(InMessageLog)
 	, Schema(InSchema)
-	, UUIDCounter(1024)
 	, bIsUbergraph(false)
 	, bCannotBeCalledFromOtherKismet(false)
 	, NetFlags(0)
@@ -897,6 +896,12 @@ void FKismetFunctionContext::SetExternalNetNameMap(FNetNameMapping* NewMap)
 	NetNameMap = NewMap;
 }
 
+int32 FKismetFunctionContext::GetContextUniqueID()
+{
+	static int32 UUIDCounter = 1;
+	return UUIDCounter++;
+}
+
 #undef LOCTEXT_NAMESPACE
 
 //////////////////////////////////////////////////////////////////////////
diff --git a/Engine/Source/Editor/KismetCompiler/Public/KismetCompiledFunctionContext.h b/Engine/Source/Editor/KismetCompiler/Public/KismetCompiledFunctionContext.h
--- a/Engine/Source/Editor/KismetCompiler/Public/KismetCompiledFunctionContext.h
+++ b/Engine/Source/Editor/KismetCompiler/Public/KismetCompiledFunctionContext.h
@@ -65,7 +65,6 @@ public:
 	TMap<UEdGraphPin*, FBPTerminal*> NetMap;
 	TMap<UEdGraphPin*, FBPTerminal*> LiteralHackMap;
 
-	int32 UUIDCounter;
 	bool bIsUbergraph;
 	bool bCannotBeCalledFromOtherKismet;
 	bool bIsInterfaceStub;
@@ -186,11 +185,7 @@ public:
 		}
 	}
 
-
-	int32 GetContextUniqueID()
-	{
-		return UUIDCounter++;
-	}
+	KISMETCOMPILER_API int32 GetContextUniqueID();
 
 	UBlueprint* GetBlueprint()
 	{
