KismetCompiler: All classes subobjects, that wont be regenerated, are not removed from the class in FKismetCompilerContext::CleanAndSanitizeClass.

ttp#332347 BP:  CopyPropertiesForUnrelatedObjects doesn't properly handle editinlinenew subobjects

diff --git a/Engine/Source/Editor/KismetCompiler/Private/KismetCompiler.cpp b/Engine/Source/Editor/KismetCompiler/Private/KismetCompiler.cpp
--- a/Engine/Source/Editor/KismetCompiler/Private/KismetCompiler.cpp
+++ b/Engine/Source/Editor/KismetCompiler/Private/KismetCompiler.cpp
@@ -114,16 +114,38 @@ void FKismetCompilerContext::SpawnNewClass(const FString& NewClassName)
 	}
 }
 
-template<typename TOBJ> struct FCullTemplateObjectsHelper
+struct FSubobjectCollection
 { 
-	const TArray<TOBJ*>& Templates;
+private:
+	TSet<const UObject*> Collection;
 
-	FCullTemplateObjectsHelper(const TArray<TOBJ*>& InComponentTemplates) : Templates(InComponentTemplates)
-	{}
+public:
+	void AddObject(const UObject* const InObject)
+	{
+		if (InObject)
+		{
+			Collection.Add(InObject);
+			TArray<UObject*> Subobjects;
+			GetObjectsWithOuter(InObject, Subobjects, true);
+			for (auto SubObject : Subobjects)
+			{
+				Collection.Add(SubObject);
+			}
+		}
+	}
+
+	template<typename TOBJ>
+	void AddObjects(const TArray<TOBJ*>& InObjects)
+	{
+		for (const auto ObjPtr : InObjects)
+		{
+			AddObject(ObjPtr);
+		}
+	}
 
 	bool operator()(const UObject* const RemovalCandidate) const
 	{
-		return (NULL != Templates.FindByKey(RemovalCandidate));
+		return (NULL != Collection.Find(RemovalCandidate));
 	}
 };
 
@@ -156,41 +178,25 @@ void FKismetCompilerContext::CleanAndSanitizeClass(UBlueprintGeneratedClass* Cla
 	TArray<UObject*> ClassSubObjects;
 	GetObjectsWithOuter(ClassToClean, ClassSubObjects, true);
 
-	if(0 != Blueprint->ComponentTemplates.Num())
-	{
-		ClassSubObjects.RemoveAllSwap(FCullTemplateObjectsHelper<UActorComponent>(Blueprint->ComponentTemplates));
-	}
-
-	if(0 != Blueprint->Timelines.Num())
-	{
-		ClassSubObjects.RemoveAllSwap(FCullTemplateObjectsHelper<UTimelineTemplate>(Blueprint->Timelines));
-	}
-
-	if(Blueprint->SimpleConstructionScript)
-	{
-		ClassSubObjects.RemoveSwap(Blueprint->SimpleConstructionScript);
-		if(const USCS_Node* DefaultScene = Blueprint->SimpleConstructionScript->GetDefaultSceneRootNode())
-		{
-			ClassSubObjects.RemoveSwap(DefaultScene->ComponentTemplate);
-		}
-
+	{	// Save subobjects, that won't be regenerated.
+		FSubobjectCollection SubObjectsToSave;
+		SubObjectsToSave.AddObjects(Blueprint->ComponentTemplates);
+		SubObjectsToSave.AddObjects(Blueprint->Timelines);
+		if (Blueprint->SimpleConstructionScript)
 		{
-			TArray<UObject*> SCSObjectsToSave;
-			GetObjectsWithOuter(Blueprint->SimpleConstructionScript, SCSObjectsToSave, true);
-			ClassSubObjects.RemoveAllSwap(FCullTemplateObjectsHelper<UObject>(SCSObjectsToSave));
-		}
+			SubObjectsToSave.AddObject(Blueprint->SimpleConstructionScript);
+			if (const USCS_Node* DefaultScene = Blueprint->SimpleConstructionScript->GetDefaultSceneRootNode())
+			{
+				SubObjectsToSave.AddObject(DefaultScene->ComponentTemplate);
+			}
 
-		{
 			TArray<USCS_Node*> SCSNodes = Blueprint->SimpleConstructionScript->GetAllNodes();
-			TArray<UActorComponent*> ActorComponents;
-			ActorComponents.Reserve(SCSNodes.Num());
-			for(auto Iter = SCSNodes.CreateIterator(); Iter; ++Iter)
+			for (auto SCSNode : SCSNodes)
 			{
-				USCS_Node* SCSNode = *Iter;
-				ActorComponents.Add(SCSNode->ComponentTemplate);
+				SubObjectsToSave.AddObject(SCSNode->ComponentTemplate);
 			}
-			ClassSubObjects.RemoveAllSwap(FCullTemplateObjectsHelper<UActorComponent>(ActorComponents));
 		}
+		ClassSubObjects.RemoveAllSwap(SubObjectsToSave);
 	}
 
 	for( auto SubObjIt = ClassSubObjects.CreateIterator(); SubObjIt; ++SubObjIt )
