Merging UE4-Pretest @ 2042161 to UE4

diff --git a/Engine/Source/Editor/KismetCompiler/Private/KismetCompilerModule.cpp b/Engine/Source/Editor/KismetCompiler/Private/KismetCompilerModule.cpp
--- a/Engine/Source/Editor/KismetCompiler/Private/KismetCompilerModule.cpp
+++ b/Engine/Source/Editor/KismetCompiler/Private/KismetCompilerModule.cpp
@@ -42,6 +42,44 @@ public:
 	}
 };
 
+struct FPrintCompilationSummaryHelper
+{
+	static void Print(double CompileStartTime, double FinishTime, const FString Name, bool bPrintResultSuccess, FCompilerResultsLog& Results)
+	{
+
+		FNumberFormattingOptions TimeFormat;
+	  TimeFormat.MaximumFractionalDigits = 2;
+	  TimeFormat.MinimumFractionalDigits = 2;
+	  TimeFormat.MaximumIntegralDigits = 4;
+	  TimeFormat.MinimumIntegralDigits = 4;
+	  TimeFormat.UseGrouping = false;
+  
+	  FFormatNamedArguments Args;
+	  Args.Add(TEXT("Time"), FText::AsNumber(FinishTime - GStartTime, &TimeFormat));
+	  Args.Add(TEXT("BlueprintName"), FText::FromString(Name));
+	  Args.Add(TEXT("NumWarnings"), Results.NumWarnings);
+	  Args.Add(TEXT("NumErrors"), Results.NumErrors);
+	  Args.Add(TEXT("TotalMilliseconds"), (int)((FinishTime - CompileStartTime) * 1000));
+
+		if (Results.NumErrors > 0)
+	  {
+		  FString FailMsg = FText::Format(LOCTEXT("CompileFailed", "[{Time}] Compile of {BlueprintName} failed. {NumErrors} Fatal Issue(s) {NumWarnings} Warning(s) [in {TotalMilliseconds} ms]"), Args).ToString();
+		  Results.Warning(*FailMsg);
+	  }
+	  else if(Results.NumWarnings > 0)
+	  {
+		  FString WarningMsg = FText::Format(LOCTEXT("CompileWarning", "[{Time}] Compile of {BlueprintName} successful, but with {NumWarnings} Warning(s) [in {TotalMilliseconds} ms]"), Args).ToString();
+		  Results.Warning(*WarningMsg);
+	  }
+	  else if(bPrintResultSuccess)
+	  {
+		  FString SuccessMsg = FText::Format(LOCTEXT("CompileSuccess", "[{Time}] Compile of {BlueprintName} successful! [in {TotalMilliseconds} ms]"), Args).ToString();
+		  Results.Note(*SuccessMsg);
+	  }
+	}
+};
+
+
 // Compiles a blueprint.
 void FKismet2CompilerModule::CompileBlueprintInner(class UBlueprint* Blueprint, bool bPrintResultSuccess, const FKismetCompilerOptions& CompileOptions, FCompilerResultsLog& Results, TArray<UObject*>* ObjLoaded)
 {
@@ -49,7 +87,7 @@ void FKismet2CompilerModule::CompileBlueprintInner(class UBlueprint* Blueprint,
 
 	Blueprint->CurrentMessageLog = &Results;
 
-	double CompileStartTime = FPlatformTime::Seconds();
+	const double CompileStartTime = FPlatformTime::Seconds();
 
 	// Early out if blueprint parent is missing
 	if (Blueprint->ParentClass == NULL)
@@ -89,39 +127,19 @@ void FKismet2CompilerModule::CompileBlueprintInner(class UBlueprint* Blueprint,
 		}
 	}
 
-	double FinishTime = FPlatformTime::Seconds();
+	const double FinishTime = FPlatformTime::Seconds();
+	FPrintCompilationSummaryHelper::Print(CompileStartTime, FinishTime, Blueprint->GetName(), bPrintResultSuccess, Results);
 
-	FNumberFormattingOptions TimeFormat;
-	TimeFormat.MaximumFractionalDigits = 2;
-	TimeFormat.MinimumFractionalDigits = 2;
-	TimeFormat.MaximumIntegralDigits = 4;
-	TimeFormat.MinimumIntegralDigits = 4;
-	TimeFormat.UseGrouping = false;
+	Blueprint->CurrentMessageLog = NULL;
+}
 
-	FFormatNamedArguments Args;
-	Args.Add(TEXT("Time"), FText::AsNumber(FinishTime - GStartTime, &TimeFormat));
-	Args.Add(TEXT("BlueprintName"), FText::FromString(Blueprint->GetName()));
-	Args.Add(TEXT("NumWarnings"), Results.NumWarnings);
-	Args.Add(TEXT("NumErrors"), Results.NumErrors);
-	Args.Add(TEXT("TotalMilliseconds"), (int)((FinishTime - CompileStartTime) * 1000));
 
-	if (Results.NumErrors > 0)
-	{
-		FString FailMsg = FText::Format(LOCTEXT("CompileFailed", "[{Time}] Compile of {BlueprintName} failed. {NumErrors} Fatal Issue(s) {NumWarnings} Warning(s) [in {TotalMilliseconds} ms]"), Args).ToString();
-		Results.Warning(*FailMsg);
-	}
-	else if(Results.NumWarnings > 0)
-	{
-		FString WarningMsg = FText::Format(LOCTEXT("CompileWarning", "[{Time}] Compile of {BlueprintName} successful, but with {NumWarnings} Warning(s) [in {TotalMilliseconds} ms]"), Args).ToString();
-		Results.Warning(*WarningMsg);
-	}
-	else if(bPrintResultSuccess)
-	{
-		FString SuccessMsg = FText::Format(LOCTEXT("CompileSuccess", "[{Time}] Compile of {BlueprintName} successful! [in {TotalMilliseconds} ms]"), Args).ToString();
-		Results.Note(*SuccessMsg);
-	}
-
-	Blueprint->CurrentMessageLog = NULL;
+void FKismet2CompilerModule::CompileStructure(class UUserDefinedStruct* Struct, FCompilerResultsLog& Results)
+{
+	const double CompileStartTime = FPlatformTime::Seconds();
+	FUserDefinedStructureCompilerUtils::CompileStruct(Struct, Results, true);
+	const double FinishTime = FPlatformTime::Seconds();
+	FPrintCompilationSummaryHelper::Print(CompileStartTime, FinishTime, Struct->GetName(), true, Results);
 }
 
 // Compiles a blueprint.
@@ -131,13 +149,6 @@ void FKismet2CompilerModule::CompileBlueprint(class UBlueprint* Blueprint, const
 
 	const bool bIsBrandNewBP = (Blueprint->SkeletonGeneratedClass == NULL) && (Blueprint->GeneratedClass == NULL) && (Blueprint->ParentClass != NULL) && !CompileOptions.bIsDuplicationInstigated;
 
-	if ((CompileOptions.CompileType == EKismetCompileType::StructuresOnly )
-		|| (CompileOptions.CompileType == EKismetCompileType::Full))
-	{
-		FUserDefinedStructureCompilerUtils::CompileStructs(Blueprint, Results, false);
-	}
-
-	if (CompileOptions.CompileType != EKismetCompileType::StructuresOnly)
 	{
 		FBlueprintCompileReinstancer SkeletonReinstancer(Blueprint->SkeletonGeneratedClass);
 
diff --git a/Engine/Source/Editor/KismetCompiler/Private/KismetCompilerVMBackend.cpp b/Engine/Source/Editor/KismetCompiler/Private/KismetCompilerVMBackend.cpp
--- a/Engine/Source/Editor/KismetCompiler/Private/KismetCompilerVMBackend.cpp
+++ b/Engine/Source/Editor/KismetCompiler/Private/KismetCompilerVMBackend.cpp
@@ -452,7 +452,7 @@ public:
 					int32 StructSize = Struct->GetStructureSize() * StructProperty->ArrayDim;
 					uint8* StructData = (uint8*)FMemory_Alloca(StructSize);
 					StructProperty->InitializeValue(StructData);
-					if(!FStructureEditorUtils::Fill_MakeStructureDefaultValue(Cast<UBlueprintGeneratedStruct>(Struct), StructData))
+					if(!FStructureEditorUtils::Fill_MakeStructureDefaultValue(Cast<UUserDefinedStruct>(Struct), StructData))
 					{
 						UE_LOG(LogK2Compiler, Warning, TEXT("MakeStructureDefaultValue parsing error. Property: %s, Struct: %s"), *StructProperty->GetName(), *Struct->GetName());
 					}
diff --git a/Engine/Source/Editor/KismetCompiler/Private/UserDefinedStructureCompilerUtils.cpp b/Engine/Source/Editor/KismetCompiler/Private/UserDefinedStructureCompilerUtils.cpp
--- a/Engine/Source/Editor/KismetCompiler/Private/UserDefinedStructureCompilerUtils.cpp
+++ b/Engine/Source/Editor/KismetCompiler/Private/UserDefinedStructureCompilerUtils.cpp
@@ -11,23 +11,24 @@
 struct FUserDefinedStructureCompilerInner
 {
 	static void ReplaceStructWithTempDuplicate(
-		UBlueprintGeneratedStruct* StructureToReinstance, 
+		UUserDefinedStruct* StructureToReinstance, 
 		TSet<UBlueprint*>& BlueprintsToRecompile,
-		TArray<UBlueprintGeneratedStruct*>& ChangedStructs)
+		TArray<UUserDefinedStruct*>& ChangedStructs)
 	{
 		if (StructureToReinstance)
 		{
-			UBlueprintGeneratedStruct* DuplicatedStruct = NULL;
+			UUserDefinedStruct* DuplicatedStruct = NULL;
 			{
 				const FString ReinstancedName = FString::Printf(TEXT("STRUCT_REINST_%s"), *StructureToReinstance->GetName());
-				const FName UniqueName = MakeUniqueObjectName(GetTransientPackage(), UBlueprintGeneratedStruct::StaticClass(), FName(*ReinstancedName));
+				const FName UniqueName = MakeUniqueObjectName(GetTransientPackage(), UUserDefinedStruct::StaticClass(), FName(*ReinstancedName));
 
 				const bool OldIsDuplicatingClassForReinstancing = GIsDuplicatingClassForReinstancing;
 				GIsDuplicatingClassForReinstancing = true;
-				DuplicatedStruct = (UBlueprintGeneratedStruct*)StaticDuplicateObject(StructureToReinstance, GetTransientPackage(), *UniqueName.ToString(), ~RF_Transactional); 
+				DuplicatedStruct = (UUserDefinedStruct*)StaticDuplicateObject(StructureToReinstance, GetTransientPackage(), *UniqueName.ToString(), ~RF_Transactional); 
 				GIsDuplicatingClassForReinstancing = OldIsDuplicatingClassForReinstancing;
 			}
-			DuplicatedStruct->Status = EBlueprintStructureStatus::BSS_Duplicate;
+			DuplicatedStruct->PrimaryStruct = StructureToReinstance;
+			DuplicatedStruct->Status = EUserDefinedStructureStatus::UDSS_Duplicate;
 			DuplicatedStruct->SetFlags(RF_Transient);
 			DuplicatedStruct->AddToRoot();
 
@@ -35,34 +36,30 @@ struct FUserDefinedStructureCompilerInner
 			{
 				if (StructProperty && (StructureToReinstance == StructProperty->Struct))
 				{
-					UBlueprint* FoundBlueprint = NULL;
 					if (auto OwnerClass = Cast<UBlueprintGeneratedClass>(StructProperty->GetOwnerClass()))
 					{
-						FoundBlueprint = Cast<UBlueprint>(OwnerClass->ClassGeneratedBy);
+						if (UBlueprint* FoundBlueprint = Cast<UBlueprint>(OwnerClass->ClassGeneratedBy))
+						{
+							BlueprintsToRecompile.Add(FoundBlueprint);
+							StructProperty->Struct = DuplicatedStruct;
+						}
 					}
-					else if (auto OwnerStruct = Cast<UBlueprintGeneratedStruct>(StructProperty->GetOwnerStruct()))
+					else if (auto OwnerStruct = Cast<UUserDefinedStruct>(StructProperty->GetOwnerStruct()))
 					{
 						check(OwnerStruct != DuplicatedStruct);
 						const bool bValidStruct = (OwnerStruct->GetOutermost() != GetTransientPackage())
 							&& !OwnerStruct->HasAnyFlags(RF_PendingKill)
-							&& (EBlueprintStructureStatus::BSS_Duplicate != OwnerStruct->Status.GetValue());
+							&& (EUserDefinedStructureStatus::UDSS_Duplicate != OwnerStruct->Status.GetValue());
 
 						if (bValidStruct)
 						{
-							FoundBlueprint = OwnerStruct->StructGeneratedBy;
 							ChangedStructs.AddUnique(OwnerStruct);
+							StructProperty->Struct = DuplicatedStruct;
 						}
-						
 					}
 					else
 					{
-						UE_LOG(LogK2Compiler, Warning, TEXT("ReplaceStructWithTempDuplicate unknown owner"));
-					}
-
-					if (NULL != FoundBlueprint)
-					{
-						StructProperty->Struct = DuplicatedStruct;
-						BlueprintsToRecompile.Add(FoundBlueprint);
+						UE_LOG(LogK2Compiler, Error, TEXT("ReplaceStructWithTempDuplicate unknown owner"));
 					}
 				}
 			}
@@ -71,34 +68,21 @@ struct FUserDefinedStructureCompilerInner
 		}
 	}
 
-	static UBlueprintGeneratedStruct* CreateNewStruct(FName Name, class UBlueprint* Blueprint)
-	{
-		check(NULL != Blueprint);
-		check(NULL != Blueprint->GetOutermost());
-		check(Name != NAME_None);
-		UBlueprintGeneratedStruct* Struct = NewNamedObject<UBlueprintGeneratedStruct>(Blueprint->GetOutermost(), Name, RF_Public|RF_Transactional /*|RF_Standalone*/);
-		check(NULL != Struct);
-		Struct->StructGeneratedBy = Blueprint;
-		return Struct;
-	}
-
-	static void CleanAndSanitizeStruct(UBlueprintGeneratedStruct* StructToClean)
+	static void CleanAndSanitizeStruct(UUserDefinedStruct* StructToClean)
 	{
 		check(StructToClean);
 
 		const FString TransientString = FString::Printf(TEXT("TRASHSTRUCT_%s"), *StructToClean->GetName());
-		const FName TransientName = MakeUniqueObjectName(GetTransientPackage(), UBlueprintGeneratedStruct::StaticClass(), FName(*TransientString));
-		UBlueprintGeneratedStruct* TransientStruct = NewNamedObject<UBlueprintGeneratedStruct>(GetTransientPackage(), TransientName, RF_Public|RF_Transient);
-
-		const bool bRecompilingOnLoad = StructToClean->StructGeneratedBy ? StructToClean->StructGeneratedBy->bIsRegeneratingOnLoad : false;
-		const ERenameFlags RenFlags = REN_DontCreateRedirectors | (bRecompilingOnLoad ? REN_ForceNoResetLoaders : 0);
+		const FName TransientName = MakeUniqueObjectName(GetTransientPackage(), UUserDefinedStruct::StaticClass(), FName(*TransientString));
+		UUserDefinedStruct* TransientStruct = NewNamedObject<UUserDefinedStruct>(GetTransientPackage(), TransientName, RF_Public|RF_Transient);
 
 		TArray<UObject*> SubObjects;
 		GetObjectsWithOuter(StructToClean, SubObjects, true);
+		SubObjects.Remove(StructToClean->EditorData);
 		for( auto SubObjIt = SubObjects.CreateIterator(); SubObjIt; ++SubObjIt )
 		{
 			UObject* CurrSubObj = *SubObjIt;
-			CurrSubObj->Rename(NULL, TransientStruct, RenFlags);
+			CurrSubObj->Rename(NULL, TransientStruct, REN_DontCreateRedirectors);
 			if( UProperty* Prop = Cast<UProperty>(CurrSubObj) )
 			{
 				FKismetCompilerUtilities::InvalidatePropertyExport(Prop);
@@ -109,11 +93,6 @@ struct FUserDefinedStructureCompilerInner
 			}
 		}
 
-		if(!bRecompilingOnLoad)
-		{
-			StructToClean->RemoveMetaData(TEXT("BlueprintType"));
-		}
-
 		StructToClean->SetSuperStruct(NULL);
 		StructToClean->Children = NULL;
 		StructToClean->Script.Empty();
@@ -125,24 +104,26 @@ struct FUserDefinedStructureCompilerInner
 		StructToClean->PropertyLink = NULL;
 	}
 
-	static void CreateVariables(UBlueprintGeneratedStruct* Struct, TArray<FBPVariableDescription>& Fields, const class UEdGraphSchema_K2* Schema, FCompilerResultsLog& MessageLog)
+	static void CreateVariables(UUserDefinedStruct* Struct, const class UEdGraphSchema_K2* Schema, FCompilerResultsLog& MessageLog)
 	{
 		check(Struct && Schema);
 
 		//FKismetCompilerUtilities::LinkAddedProperty push property to begin, so we revert the order
-		for(int32 VarDescIdx = Fields.Num() - 1; VarDescIdx >= 0; --VarDescIdx)
+		for (int32 VarDescIdx = FStructureEditorUtils::GetVarDesc(Struct).Num() - 1; VarDescIdx >= 0; --VarDescIdx)
 		{
-			FBPVariableDescription& VarDesc = Fields[VarDescIdx];
+			FStructVariableDescription& VarDesc = FStructureEditorUtils::GetVarDesc(Struct)[VarDescIdx];
 			VarDesc.bInvalidMember = true;
 
+			FEdGraphPinType VarType = VarDesc.ToPinType();
+
 			FString ErrorMsg;
-			if(!FStructureEditorUtils::CanHaveAMemberVariableOfType(Struct, VarDesc.VarType, &ErrorMsg))
+			if(!FStructureEditorUtils::CanHaveAMemberVariableOfType(Struct, VarType, &ErrorMsg))
 			{
-				MessageLog.Error(*FString::Printf(*LOCTEXT("StructureGeneric_Error", "Structure: %s Error: %s").ToString(), *Struct->GetPathName(), *ErrorMsg));
+				MessageLog.Error(*FString::Printf(*LOCTEXT("StructureGeneric_Error", "Structure: %s Error: %s").ToString(), *Struct->GetFullName(), *ErrorMsg));
 				continue;
 			}
 
-			UProperty* NewProperty = FKismetCompilerUtilities::CreatePropertyOnScope(Struct, VarDesc.VarName, VarDesc.VarType, NULL, 0, Schema, MessageLog);
+			UProperty* NewProperty = FKismetCompilerUtilities::CreatePropertyOnScope(Struct, VarDesc.VarName, VarType, NULL, 0, Schema, MessageLog);
 			if (NewProperty != NULL)
 			{
 				FKismetCompilerUtilities::LinkAddedProperty(Struct, NewProperty);
@@ -150,85 +131,73 @@ struct FUserDefinedStructureCompilerInner
 			else
 			{
 				MessageLog.Error(*FString::Printf(*LOCTEXT("VariableInvalidType_Error", "The variable %s declared in %s has an invalid type %s").ToString(),
-					*VarDesc.VarName.ToString(), *Struct->GetName(), *UEdGraphSchema_K2::TypeToString(VarDesc.VarType)));
+					*VarDesc.VarName.ToString(), *Struct->GetName(), *UEdGraphSchema_K2::TypeToString(VarType)));
 				continue;
 			}
 
-			NewProperty->SetPropertyFlags(VarDesc.PropertyFlags);
+			NewProperty->SetPropertyFlags(CPF_Edit);
 			NewProperty->SetMetaData(TEXT("FriendlyName"), *VarDesc.FriendlyName);
 			NewProperty->SetMetaData(TEXT("DisplayName"), *VarDesc.FriendlyName);
-			NewProperty->SetMetaData(TEXT("Category"), *VarDesc.Category.ToString());
-			NewProperty->RepNotifyFunc = VarDesc.RepNotifyFunc;
+			NewProperty->SetMetaData(TEXT("Category"), *VarDesc.Category);
+			NewProperty->RepNotifyFunc = NAME_None;
 
 			if (!VarDesc.DefaultValue.IsEmpty())
 			{
 				NewProperty->SetMetaData(TEXT("MakeStructureDefaultValue"), *VarDesc.DefaultValue);
 			}
-
-			// Set metadata on property
-			for (const FBPVariableMetaDataEntry& Entry : VarDesc.MetaDataArray)
-			{
-				NewProperty->SetMetaData(Entry.DataKey, *Entry.DataValue);
-			}
+			VarDesc.CurrentDefaultValue = VarDesc.DefaultValue;
 
 			VarDesc.bInvalidMember = false;
 		}
 	}
 
-	static void InnerCompileStruct(FBPStructureDescription& StructDesc, const class UEdGraphSchema_K2* K2Schema, class FCompilerResultsLog& MessageLog)
+	static void InnerCompileStruct(UUserDefinedStruct* Struct, const class UEdGraphSchema_K2* K2Schema, class FCompilerResultsLog& MessageLog)
 	{
+		check(Struct);
 		const int32 ErrorNum = MessageLog.NumErrors;
 
-		StructDesc.CompiledStruct->SetMetaData(TEXT("BlueprintType"), TEXT("true"));
+		CreateVariables(Struct, K2Schema, MessageLog);
 
-		CreateVariables(StructDesc.CompiledStruct, StructDesc.Fields, K2Schema, MessageLog);
+		Struct->Bind();
+		Struct->StaticLink(true);
 
-		StructDesc.CompiledStruct->Bind();
-		StructDesc.CompiledStruct->StaticLink(true);
+		if (Struct->GetStructureSize() <= 0)
+		{
+			MessageLog.Error(*FString::Printf(*LOCTEXT("StructurEmpty_Error", "Structure '%s' is empty ").ToString(), *Struct->GetFullName()));
+		}
 
 		const bool bNoErrorsDuringCompilation = (ErrorNum == MessageLog.NumErrors);
-		StructDesc.CompiledStruct->Status = bNoErrorsDuringCompilation ? EBlueprintStructureStatus::BSS_UpToDate : EBlueprintStructureStatus::BSS_Error;
+		Struct->Status = bNoErrorsDuringCompilation ? EUserDefinedStructureStatus::UDSS_UpToDate : EUserDefinedStructureStatus::UDSS_Error;
 	}
 
-	static bool ShouldBeCompiled(const FBPStructureDescription& StructDesc)
+	static bool ShouldBeCompiled(const UUserDefinedStruct* Struct)
 	{
-		const UBlueprintGeneratedStruct* Struct = StructDesc.CompiledStruct;
-		if (Struct && (EBlueprintStructureStatus::BSS_UpToDate == Struct->Status))
+		if (Struct && (EUserDefinedStructureStatus::UDSS_UpToDate == Struct->Status))
 		{
 			return false;
 		}
 		return true;
 	}
 
-	static void BuildDependencyMapAndCompile(TArray<UBlueprintGeneratedStruct*>& ChangedStructs, FCompilerResultsLog& MessageLog)
+	static void BuildDependencyMapAndCompile(TArray<UUserDefinedStruct*>& ChangedStructs, FCompilerResultsLog& MessageLog)
 	{
-		struct FFindStructureDescriptionPred
-		{
-			const UBlueprintGeneratedStruct* const Struct;
-			FFindStructureDescriptionPred(const UBlueprintGeneratedStruct* InStruct) : Struct(InStruct) { check(Struct); }
-			bool operator()(const FBPStructureDescription& Desc) { return (Desc.CompiledStruct == Struct); }
-		};
-
 		struct FDependencyMapEntry
 		{
-			UBlueprintGeneratedStruct* Struct;
-			TSet<UBlueprintGeneratedStruct*> StructuresToWaitFor;
+			UUserDefinedStruct* Struct;
+			TSet<UUserDefinedStruct*> StructuresToWaitFor;
 
 			FDependencyMapEntry() : Struct(NULL) {}
 
-			void Initialize(UBlueprintGeneratedStruct* ChangedStruct, TArray<UBlueprintGeneratedStruct*>& AllChangedStructs) 
+			void Initialize(UUserDefinedStruct* ChangedStruct, TArray<UUserDefinedStruct*>& AllChangedStructs) 
 			{ 
 				Struct = ChangedStruct;
-				check(Struct && Struct->StructGeneratedBy);
-				FBPStructureDescription* StructureDescription = Struct->StructGeneratedBy->UserDefinedStructures.FindByPredicate(FFindStructureDescriptionPred(Struct));
-				check(StructureDescription);
+				check(Struct);
 
 				auto Schema = GetDefault<UEdGraphSchema_K2>();
-				for (auto FieldIter = StructureDescription->Fields.CreateIterator(); FieldIter; ++FieldIter)
+				for (auto& VarDesc : FStructureEditorUtils::GetVarDesc(Struct))
 				{
-					FEdGraphPinType FieldType = (*FieldIter).VarType;
-					auto StructType = Cast<UBlueprintGeneratedStruct>(FieldType.PinSubCategoryObject.Get());
-					if (StructType && (FieldType.PinCategory == Schema->PC_Struct) && AllChangedStructs.Contains(StructType))
+					auto StructType = Cast<UUserDefinedStruct>(VarDesc.SubCategoryObject.Get());
+					if (StructType && (VarDesc.Category == Schema->PC_Struct) && AllChangedStructs.Contains(StructType))
 					{
 						StructuresToWaitFor.Add(StructType);
 					}
@@ -255,13 +224,11 @@ struct FUserDefinedStructureCompilerInner
 				}
 			}
 			check(INDEX_NONE != StructureToCompileIndex);
-			UBlueprintGeneratedStruct* Struct = DependencyMap[StructureToCompileIndex].Struct;
-			check(Struct->StructGeneratedBy);
-			FBPStructureDescription* StructureDescription = Struct->StructGeneratedBy->UserDefinedStructures.FindByPredicate(FFindStructureDescriptionPred(Struct));
-			check(StructureDescription);
+			UUserDefinedStruct* Struct = DependencyMap[StructureToCompileIndex].Struct;
+			check(Struct);
 
 			FUserDefinedStructureCompilerInner::CleanAndSanitizeStruct(Struct);
-			FUserDefinedStructureCompilerInner::InnerCompileStruct(*StructureDescription, GetDefault<UEdGraphSchema_K2>(), MessageLog);
+			FUserDefinedStructureCompilerInner::InnerCompileStruct(Struct, GetDefault<UEdGraphSchema_K2>(), MessageLog);
 
 			DependencyMap.RemoveAtSwap(StructureToCompileIndex);
 
@@ -273,45 +240,34 @@ struct FUserDefinedStructureCompilerInner
 	}
 };
 
-void FUserDefinedStructureCompilerUtils::CompileStructs(class UBlueprint* Blueprint, class FCompilerResultsLog& MessageLog, bool bForceRecompileAll)
+void FUserDefinedStructureCompilerUtils::CompileStruct(class UUserDefinedStruct* Struct, class FCompilerResultsLog& MessageLog, bool bForceRecompile)
 {
-	if (FStructureEditorUtils::StructureEditingEnabled() && Blueprint && Blueprint->UserDefinedStructures.Num())
+	if (FStructureEditorUtils::UserDefinedStructEnabled() && Struct)
 	{
 		TSet<UBlueprint*> BlueprintsToRecompile;
-		TArray<UBlueprintGeneratedStruct*> ChangedStructs; 
-		for (auto StructDescIter = Blueprint->UserDefinedStructures.CreateIterator(); StructDescIter; ++StructDescIter)
+		TArray<UUserDefinedStruct*> ChangedStructs; 
+
+		if (FUserDefinedStructureCompilerInner::ShouldBeCompiled(Struct) || bForceRecompile)
 		{
-			FBPStructureDescription& StructDesc = (*StructDescIter);
-			if (UBlueprintGeneratedStruct* Struct = StructDesc.CompiledStruct)
-			{
-				if (FUserDefinedStructureCompilerInner::ShouldBeCompiled(StructDesc) || bForceRecompileAll)
-				{
-					ChangedStructs.Add(Struct);
-					BlueprintsToRecompile.Add(Blueprint);
-				}
-			}
-			else
-			{
-				 StructDesc.CompiledStruct = FUserDefinedStructureCompilerInner::CreateNewStruct(StructDesc.Name, Blueprint);
-				 FUserDefinedStructureCompilerInner::InnerCompileStruct(StructDesc, GetDefault<UEdGraphSchema_K2>(), MessageLog);
-			}
+			ChangedStructs.Add(Struct);
 		}
 
 		for (int32 StructIdx = 0; StructIdx < ChangedStructs.Num(); ++StructIdx)
 		{
 			FUserDefinedStructureCompilerInner::ReplaceStructWithTempDuplicate(ChangedStructs[StructIdx], BlueprintsToRecompile, ChangedStructs);
-			ChangedStructs[StructIdx]->Status = EBlueprintStructureStatus::BSS_Dirty;
+			ChangedStructs[StructIdx]->Status = EUserDefinedStructureStatus::UDSS_Dirty;
 		}
 
 		// COMPILE IN PROPER ORDER
 		FUserDefinedStructureCompilerInner::BuildDependencyMapAndCompile(ChangedStructs, MessageLog);
 
+		// UPDATE ALL THINGS DEPENDENT ON COMPILED STRUCTURES
 		for (TObjectIterator<UK2Node_StructOperation> It(RF_Transient | RF_PendingKill | RF_ClassDefaultObject, true); It && ChangedStructs.Num(); ++It)
 		{
 			UK2Node_StructOperation* Node = *It;
 			if (Node && !Node->HasAnyFlags(RF_Transient|RF_PendingKill))
 			{
-				UBlueprintGeneratedStruct* StructInNode = Cast<UBlueprintGeneratedStruct>(Node->StructType);
+				UUserDefinedStruct* StructInNode = Cast<UUserDefinedStruct>(Node->StructType);
 				if (StructInNode && ChangedStructs.Contains(StructInNode))
 				{
 					if (UBlueprint* FoundBlueprint = Node->GetBlueprint())
@@ -332,7 +288,7 @@ void FUserDefinedStructureCompilerUtils::CompileStructs(class UBlueprint* Bluepr
 
 void FUserDefinedStructureCompilerUtils::DefaultUserDefinedStructs(UObject* Object, FCompilerResultsLog& MessageLog)
 {
-	if (Object && FStructureEditorUtils::StructureEditingEnabled())
+	if (Object && FStructureEditorUtils::UserDefinedStructEnabled())
 	{
 		for (TFieldIterator<UProperty> It(Object->GetClass()); It; ++It)
 		{
diff --git a/Engine/Source/Editor/KismetCompiler/Public/KismetCompilerModule.h b/Engine/Source/Editor/KismetCompiler/Public/KismetCompilerModule.h
--- a/Engine/Source/Editor/KismetCompiler/Public/KismetCompilerModule.h
+++ b/Engine/Source/Editor/KismetCompiler/Public/KismetCompilerModule.h
@@ -26,6 +26,14 @@ public:
 	 */
 	virtual void CompileBlueprint(class UBlueprint* Blueprint, const FKismetCompilerOptions& CompileOptions, FCompilerResultsLog& Results, class FBlueprintCompileReinstancer* ParentReinstancer = NULL, TArray<UObject*>* ObjLoaded = NULL)=0;
 
+	/**
+	 * Compiles a user defined structure.
+	 *
+	 * @param	Struct		The structure to compile.
+	 * @param	Results  	The results log for warnings and errors.
+	 */
+	virtual void CompileStructure(class UUserDefinedStruct* Struct, FCompilerResultsLog& Results)=0;
+
 	/**
 	 * Attempts to recover a corrupted blueprint package.
 	 *
@@ -57,6 +65,7 @@ class FKismet2CompilerModule : public IKismetCompilerInterface
 public:
 	// Implementation of the IKismetCompilerInterface
 	virtual void CompileBlueprint(class UBlueprint* Blueprint, const FKismetCompilerOptions& CompileOptions, FCompilerResultsLog& Results, class FBlueprintCompileReinstancer* ParentReinstancer = NULL, TArray<UObject*>* ObjLoaded = NULL) OVERRIDE;
+	virtual void CompileStructure(class UUserDefinedStruct* Struct, FCompilerResultsLog& Results) OVERRIDE;
 	virtual void RecoverCorruptedBlueprint(class UBlueprint* Blueprint) OVERRIDE;
 	virtual void RemoveBlueprintGeneratedClasses(class UBlueprint* Blueprint) OVERRIDE;
 	virtual TArray<FBlueprintCompileDelegate>& GetCompilers() OVERRIDE { return Compilers; }
diff --git a/Engine/Source/Editor/KismetCompiler/Public/UserDefinedStructureCompilerUtils.h b/Engine/Source/Editor/KismetCompiler/Public/UserDefinedStructureCompilerUtils.h
--- a/Engine/Source/Editor/KismetCompiler/Public/UserDefinedStructureCompilerUtils.h
+++ b/Engine/Source/Editor/KismetCompiler/Public/UserDefinedStructureCompilerUtils.h
@@ -7,9 +7,10 @@
 class FUserDefinedStructureCompilerUtils
 {
 public:
-	// ASSUMPTION, no structure needs to be renamed or removed
-	static void CompileStructs(class UBlueprint* Blueprint, class FCompilerResultsLog& MessageLog, bool bForceRecompileAll);
+	// ASSUMPTION, structure doesn't need to be renamed or removed
+	static void CompileStruct(class UUserDefinedStruct* Struct, class FCompilerResultsLog& MessageLog, bool bForceRecompile);
 
+	// Fill all variables of given object, that are UserDefinedStructure type, with default values.
 	// Default values for member variables in User Defined Structure are stored in meta data "MakeStructureDefaultValue"
 	static void DefaultUserDefinedStructs(UObject* Object, FCompilerResultsLog& MessageLog);
 };
diff --git a/Engine/Source/Editor/KismetWidgets/Private/SPinTypeSelector.cpp b/Engine/Source/Editor/KismetWidgets/Private/SPinTypeSelector.cpp
--- a/Engine/Source/Editor/KismetWidgets/Private/SPinTypeSelector.cpp
+++ b/Engine/Source/Editor/KismetWidgets/Private/SPinTypeSelector.cpp
@@ -131,7 +131,7 @@ TSharedRef<ITableRow> SPinTypeSelector::GenerateTypeTreeRow(FPinTypeTreeItem InI
 	const FString Description = InItem->GetDescription();
 
 	// Determine the best icon the to represents this item
-	const FSlateBrush* IconBrush = GetIconFromPin( InItem->PinType );
+	const FSlateBrush* IconBrush = GetIconFromPin(InItem->GetPinType(false));
 
 	// Use tooltip if supplied, otherwise just repeat description
 	const FString OrgTooltip = InItem->GetToolTip();
@@ -147,7 +147,7 @@ TSharedRef<ITableRow> SPinTypeSelector::GenerateTypeTreeRow(FPinTypeTreeItem InI
 			[
 				SNew(SImage)
 				.Image(IconBrush)
-				.ColorAndOpacity( Schema->GetPinTypeColor(InItem->PinType) )
+				.ColorAndOpacity(Schema->GetPinTypeColor(InItem->GetPinType(false)))
 				.Visibility( InItem->bReadOnly ? EVisibility::Collapsed : EVisibility::Visible )
 			]
 			+SHorizontalBox::Slot()
@@ -181,10 +181,12 @@ void SPinTypeSelector::OnTypeSelectionChanged(FPinTypeTreeItem Selection, ESelec
 			//Call delegate in order to notify pin type change is about to happen
 			OnTypePreChanged.ExecuteIfBound(NewTargetPinType);
 
+			const FEdGraphPinType& SelectionPinType = Selection->GetPinType(true);
+
 			// Change the pin's type
-			NewTargetPinType.PinCategory = Selection->PinType.PinCategory;
-			NewTargetPinType.PinSubCategory = Selection->PinType.PinSubCategory;
-			NewTargetPinType.PinSubCategoryObject = Selection->PinType.PinSubCategoryObject;
+			NewTargetPinType.PinCategory = SelectionPinType.PinCategory;
+			NewTargetPinType.PinSubCategory = SelectionPinType.PinSubCategory;
+			NewTargetPinType.PinSubCategoryObject = SelectionPinType.PinSubCategoryObject;
 
 			TypeComboButton->SetIsOpen(false);
 
diff --git a/Engine/Source/Editor/LandscapeEditor/Private/LandscapeEditorModule.cpp b/Engine/Source/Editor/LandscapeEditor/Private/LandscapeEditorModule.cpp
--- a/Engine/Source/Editor/LandscapeEditor/Private/LandscapeEditorModule.cpp
+++ b/Engine/Source/Editor/LandscapeEditor/Private/LandscapeEditorModule.cpp
@@ -4,7 +4,7 @@
 #include "LandscapeEditorCommands.h"
 #include "LandscapeEdMode.h"
 #include "LandscapeEditorClasses.h"
-#include "LandscapeEditor.generated.inl"
+
 
 #include "LandscapeEditorDetails.h"
 #include "LandscapeEditorDetailCustomizations.h"
diff --git a/Engine/Source/Editor/Matinee/Classes/InterpTrackHelper.h b/Engine/Source/Editor/Matinee/Classes/InterpTrackHelper.h
--- a/Engine/Source/Editor/Matinee/Classes/InterpTrackHelper.h
+++ b/Engine/Source/Editor/Matinee/Classes/InterpTrackHelper.h
@@ -50,5 +50,11 @@ class MATINEE_API UInterpTrackHelper : public UObject
 	* @param KeyIndex	The index of the keyframe that as just added.  This is the index returned by AddKeyframe.
 	*/
 	virtual void  PostCreateKeyframe( UInterpTrack *Track, int32 KeyIndex ) const { }
+
+	/** Returns the name of the new keyframe that has been added, valid while in the process of a rename operation */
+	FName GetKeyframeAddDataName() const { return KeyframeAddDataName; }
+
+protected:
+	static FName	KeyframeAddDataName;
 };
 
diff --git a/Engine/Source/Editor/Matinee/Matinee.Build.cs b/Engine/Source/Editor/Matinee/Matinee.Build.cs
--- a/Engine/Source/Editor/Matinee/Matinee.Build.cs
+++ b/Engine/Source/Editor/Matinee/Matinee.Build.cs
@@ -37,7 +37,8 @@ public class Matinee : ModuleRules
 				"UnrealEd",
 				"RenderCore",
 				"AssetRegistry",
-                "ContentBrowser"
+				"ContentBrowser",
+				"BlueprintGraph"
 			}
 			);
 
diff --git a/Engine/Source/Editor/Matinee/Private/Matinee.cpp b/Engine/Source/Editor/Matinee/Private/Matinee.cpp
--- a/Engine/Source/Editor/Matinee/Private/Matinee.cpp
+++ b/Engine/Source/Editor/Matinee/Private/Matinee.cpp
@@ -13,7 +13,7 @@
 #include "IDistCurveEditor.h"
 
 #include "MatineeClasses.h"
-#include "Matinee.generated.inl"
+
 
 #include "CameraController.h"
 #include "MatineeConstants.h"
diff --git a/Engine/Source/Editor/Matinee/Private/MatineeMenus.cpp b/Engine/Source/Editor/Matinee/Private/MatineeMenus.cpp
--- a/Engine/Source/Editor/Matinee/Private/MatineeMenus.cpp
+++ b/Engine/Source/Editor/Matinee/Private/MatineeMenus.cpp
@@ -6,6 +6,7 @@
 
 #include "MatineeModule.h"
 #include "Matinee.h"
+#include "MatineeDelegates.h"
 
 #include "EngineInterpolationClasses.h"
 #include "CameraController.h"
@@ -2530,14 +2531,17 @@ void FMatinee::OnContextRenameEventKeyTextCommitted(const FText& InText, ETextCo
 		FName NewEventName = FName( *TempString );
 
 		// If this Event name is already in use- disallow it
-		TArray<FName> CurrentEventNames;
-		IData->GetAllEventNames(CurrentEventNames);
-		if( CurrentEventNames.Contains(NewEventName) )
+		if( IData->IsEventName( NewEventName ) )
 		{
 			FMessageDialog::Open( EAppMsgType::Ok, NSLOCTEXT("UnrealEd", "Error_EventNameInUse", "Sorry - Event name already in use.") );
 			return;
 		}
 	
+		InterpEdTrans->BeginSpecial( LOCTEXT( "EventRename", "Event Rename" ) );
+
+		MatineeActor->Modify();
+		IData->Modify();
+
 		// Then go through all keys, changing those with this name to the new one.
 		for(int32 i=0; i<IData->InterpGroups.Num(); i++)
 		{
@@ -2547,18 +2551,30 @@ void FMatinee::OnContextRenameEventKeyTextCommitted(const FText& InText, ETextCo
 				UInterpTrackEvent* EventTrack = Cast<UInterpTrackEvent>( Group->InterpTracks[j] );
 				if(EventTrack)
 				{
+					bool bModified = false;
 					for(int32 k=0; k<EventTrack->EventTrack.Num(); k++)
 					{
-						if(EventTrack->EventTrack[k].EventName == EventNameToChange)
+						FEventTrackKey& Key = EventTrack->EventTrack[k];
+						if(Key.EventName == EventNameToChange)
 						{
-							EventTrack->EventTrack[k].EventName = NewEventName;
+							if ( !bModified )
+							{
+								EventTrack->Modify();
+								bModified = true;
+							}
+							Key.EventName = NewEventName;
 						}	
 					}
 				}			
 			}
 		}
 
+		// Fire a delegate so other places that use the name can also update
+		FMatineeDelegates::Get().OnEventKeyframeRenamed.Broadcast( MatineeActor, EventNameToChange, NewEventName );
+
 		IData->UpdateEventNames();
+
+		InterpEdTrans->EndSpecial();
 	}
 }
 
diff --git a/Engine/Source/Editor/Matinee/Private/MatineeTools.cpp b/Engine/Source/Editor/Matinee/Private/MatineeTools.cpp
--- a/Engine/Source/Editor/Matinee/Private/MatineeTools.cpp
+++ b/Engine/Source/Editor/Matinee/Private/MatineeTools.cpp
@@ -10,6 +10,7 @@
 
 #include "MatineeConstants.h"
 #include "MatineeViewSaveData.h"
+#include "MatineeDelegates.h"
 
 #include "SoundDefinitions.h"
 
@@ -1178,15 +1179,40 @@ void FMatinee::DeleteSelectedKeys(bool bDoTransaction)
 				ModifiedTracks.Add(Track);
 			}
 		}
+			
 
-		// If this is an event key - we update the connectors later.
-		if(Track->IsA(UInterpTrackEvent::StaticClass()))
+		FName OldKeyName = NAME_None;
+		if(const UInterpTrackEvent* TrackEvent = Cast< UInterpTrackEvent >( Track ))
 		{
+			// If this is an event key - we update the connectors later.
 			bRemovedEventKeys = true;
+
+
+			// Take a copy of the key name before it's removed
+			if( TrackEvent->EventTrack.IsValidIndex( SelKey.KeyIndex ))
+			{
+				OldKeyName = TrackEvent->EventTrack[SelKey.KeyIndex].EventName;
+			}
 		}
-			
+
+		
 		Track->RemoveKeyframe(SelKey.KeyIndex);
 
+
+		// If we have a valid name, check to see if it's last event key to be removed with this name? 
+		if( !OldKeyName.IsNone() )
+		{
+			const bool bCommonName = IData->IsEventName( OldKeyName );
+			if( !bCommonName )
+			{
+				// Fire a delegate so other places that use the name can also update
+				TArray<FName> KeyNames;
+				KeyNames.Add( OldKeyName );
+				FMatineeDelegates::Get().OnEventKeyframeRemoved.Broadcast( MatineeActor, KeyNames );
+			}
+		}
+
+
 		// If any other keys in the selection are on the same track but after the one we just deleted, decrement the index to correct it.
 		for(int32 j=0; j<Opt->SelectedKeys.Num(); j++)
 		{
@@ -2221,8 +2247,27 @@ void FMatinee::FinishAddKey(UInterpTrack* Track, bool bCommitKeys)
 			// Add key at current time, snapped to the grid if its on.
 			int32 NewKeyIndex = Track->AddKeyframe( fKeyTime, TrInst, InitialInterpMode );
 			check( NewKeyIndex != INDEX_NONE );
+
+
+			// Check to see if this is going to be the event first key to have this name
+			bool bCommonName = true;
+			FName NewKeyName = TrackHelper->GetKeyframeAddDataName();
+			if( !NewKeyName.IsNone() )
+			{
+				bCommonName = IData->IsEventName( NewKeyName );
+			}				
+
+
 			TrackHelper->PostCreateKeyframe( Track, NewKeyIndex );
 			TrackToNewKeyIndexMap.Add( Track, NewKeyIndex );
+
+
+			// Is this the first event key to be added with this name?
+			if( !bCommonName )
+			{
+				// Fire a delegate so other places that use the name can also update
+				FMatineeDelegates::Get().OnEventKeyframeAdded.Broadcast( MatineeActor, NewKeyName, NewKeyIndex );
+			}
 		}
 
 		InterpEdTrans->EndSpecial();
@@ -2868,6 +2913,10 @@ void FMatinee::DeleteSelectedTracks()
 	// Deselect everything.
 	ClearKeySelection();
 
+	// Take a copy of all the valid event names
+	TArray<FName> OldEventNames;
+	MatineeActor->MatineeData->GetAllEventNames(OldEventNames);
+
 	for( FSelectedTrackIterator TrackIt(GetSelectedTrackIterator()); TrackIt; ++TrackIt )
 	{
 		UInterpTrack* ActiveTrack = *TrackIt;
@@ -2917,6 +2966,28 @@ void FMatinee::DeleteSelectedTracks()
 		}
 	}
 
+	IData->UpdateEventNames();
+
+	// Take another copy of all the valid event names
+	TArray<FName> RemainingEventNames;
+	MatineeActor->MatineeData->GetAllEventNames(RemainingEventNames);
+
+	// Check to see which event name no longer exist
+	TArray<FName> RemovedEventNames;
+	for( int32 EventNameIndex = 0; EventNameIndex < OldEventNames.Num(); ++EventNameIndex )
+	{
+		const FName& OldEventName = OldEventNames[EventNameIndex];
+		if ( !RemainingEventNames.Contains(OldEventName) )
+		{
+			RemovedEventNames.Add(OldEventName);
+		}
+	}
+	if ( RemovedEventNames.Num() > 0 )
+	{
+		// Fire a delegate so other places that use the name can also update
+		FMatineeDelegates::Get().OnEventKeyframeRemoved.Broadcast( MatineeActor, RemovedEventNames );
+	}
+
 	InterpEdTrans->EndSpecial();
 
 	// Update the curve editor
@@ -2946,6 +3017,10 @@ void FMatinee::DeleteSelectedGroups()
 	// Deselect everything.
 	ClearKeySelection();
 
+	// Take a copy of all the valid event names
+	TArray<FName> OldEventNames;
+	MatineeActor->MatineeData->GetAllEventNames(OldEventNames);
+
 	for( FSelectedGroupIterator GroupIt(GetSelectedGroupIterator()); GroupIt; ++GroupIt )
 	{
 		UInterpGroup* GroupToDelete = *GroupIt;
@@ -3016,6 +3091,28 @@ void FMatinee::DeleteSelectedGroups()
 		GroupIt.RemoveCurrent();
 	}
 
+	IData->UpdateEventNames();
+
+	// Take another copy of all the valid event names
+	TArray<FName> RemainingEventNames;
+	MatineeActor->MatineeData->GetAllEventNames(RemainingEventNames);
+
+	// Check to see which event name no longer exist
+	TArray<FName> RemovedEventNames;
+	for( int32 EventNameIndex = 0; EventNameIndex < OldEventNames.Num(); ++EventNameIndex )
+	{
+		const FName& OldEventName = OldEventNames[EventNameIndex];
+		if ( !RemainingEventNames.Contains(OldEventName) )
+		{
+			RemovedEventNames.Add(OldEventName);
+		}
+	}
+	if ( RemovedEventNames.Num() > 0 )
+	{
+		// Fire a delegate so other places that use the name can also update
+		FMatineeDelegates::Get().OnEventKeyframeRemoved.Broadcast( MatineeActor, RemovedEventNames );
+	}
+
 	// Tell curve editor stuff might have changed.
 	CurveEd->CurveChanged();
 
diff --git a/Engine/Source/Editor/Matinee/Private/MatineeTrackHelpers.cpp b/Engine/Source/Editor/Matinee/Private/MatineeTrackHelpers.cpp
--- a/Engine/Source/Editor/Matinee/Private/MatineeTrackHelpers.cpp
+++ b/Engine/Source/Editor/Matinee/Private/MatineeTrackHelpers.cpp
@@ -11,6 +11,8 @@
 
 #define LOCTEXT_NAMESPACE "MatineeTrackHelpers"
 
+FName	UInterpTrackHelper::KeyframeAddDataName = NAME_None;
+
 UInterpTrackHelper::UInterpTrackHelper(const class FPostConstructInitializeProperties& PCIP)
 	: Super(PCIP)
 {
@@ -41,7 +43,6 @@ AActor* UInterpTrackHelper::GetGroupActor(const UInterpTrack* Track) const
 
 // Common FName used just for storing name information while adding Keyframes to tracks.
 static UAnimSequence*	KeyframeAddAnimSequence = NULL;
-static FName			KeyframeAddDataName = NAME_None;
 static USoundBase*		KeyframeAddSound = NULL;
 static FName			TrackAddPropName = NAME_None;
 static FName			AnimSlotName = NAME_None;
@@ -432,6 +433,7 @@ void  UMatineeTrackEventHelper::PostCreateKeyframe( UInterpTrack *Track, int32 K
 	// Update AllEventNames array now we have given it a name
 	UInterpGroup* Group = CastChecked<UInterpGroup>( EventTrack->GetOuter() );
 	UInterpData* IData = CastChecked<UInterpData>( Group->GetOuter() );
+	IData->Modify();
 	IData->UpdateEventNames();
 
 	KeyframeAddDataName = NAME_None;
diff --git a/Engine/Source/Editor/Matinee/Private/MatineeTransaction.cpp b/Engine/Source/Editor/Matinee/Private/MatineeTransaction.cpp
--- a/Engine/Source/Editor/Matinee/Private/MatineeTransaction.cpp
+++ b/Engine/Source/Editor/Matinee/Private/MatineeTransaction.cpp
@@ -3,6 +3,8 @@
 #include "MatineeModule.h"
 #include "MatineeOptions.h"
 #include "MatineeTransaction.h"
+#include "Editor/BlueprintGraph/Classes/K2Node.h"
+#include "Editor/BlueprintGraph/Classes/K2Node_MatineeController.h"
 
 /*-----------------------------------------------------------------------------
 	FMatineeTransaction
@@ -17,7 +19,8 @@ void FMatineeTransaction::SaveObject( UObject* Object )
 		Object->IsA( UInterpTrack::StaticClass() ) ||
 		Object->IsA( UInterpGroupInst::StaticClass() ) ||
 		Object->IsA( UInterpTrackInst::StaticClass() ) ||
-		Object->IsA( UMatineeOptions::StaticClass() ) )
+		Object->IsA( UMatineeOptions::StaticClass() ) ||
+		Object->IsA( UK2Node_MatineeController::StaticClass() ) )
 	{
 		// Save the object.
 		new( Records )FObjectRecord( this, Object, NULL, 0, 0, 0, 0, NULL, NULL );
diff --git a/Engine/Source/Editor/NewsFeed/Private/NewsFeedClasses.cpp b/Engine/Source/Editor/NewsFeed/Private/NewsFeedClasses.cpp
--- a/Engine/Source/Editor/NewsFeed/Private/NewsFeedClasses.cpp
+++ b/Engine/Source/Editor/NewsFeed/Private/NewsFeedClasses.cpp
@@ -5,8 +5,6 @@
 =============================================================================*/
 
 #include "NewsFeedPrivatePCH.h"
-#include "NewsFeed.generated.inl"
-
 
 /* UNewsFeedSettings interface
  *****************************************************************************/
diff --git a/Engine/Source/Editor/OnlineBlueprintSupport/Private/OnlineBlueprintSupportModule.cpp b/Engine/Source/Editor/OnlineBlueprintSupport/Private/OnlineBlueprintSupportModule.cpp
--- a/Engine/Source/Editor/OnlineBlueprintSupport/Private/OnlineBlueprintSupportModule.cpp
+++ b/Engine/Source/Editor/OnlineBlueprintSupport/Private/OnlineBlueprintSupportModule.cpp
@@ -2,6 +2,6 @@
 
 #include "OnlineBlueprintSupportPrivatePCH.h"
 #include "ModuleManager.h"
-#include "OnlineBlueprintSupport.generated.inl"
+
 
 IMPLEMENT_MODULE( FDefaultModuleImpl, OnlineBlueprintSupport );
diff --git a/Engine/Source/Editor/PhAT/Private/PhAT.cpp b/Engine/Source/Editor/PhAT/Private/PhAT.cpp
--- a/Engine/Source/Editor/PhAT/Private/PhAT.cpp
+++ b/Engine/Source/Editor/PhAT/Private/PhAT.cpp
@@ -13,7 +13,7 @@
 #include "PhATEdSkeletalMeshComponent.h"
 #include "PhAT.h"
 
-#include "PhAT.generated.inl"
+
 
 #include "Editor/WorkspaceMenuStructure/Public/WorkspaceMenuStructureModule.h"
 #include "Editor/PropertyEditor/Public/PropertyEditorModule.h"
diff --git a/Engine/Source/Editor/PropertyEditor/Private/CustomChildBuilder.cpp b/Engine/Source/Editor/PropertyEditor/Private/CustomChildBuilder.cpp
--- a/Engine/Source/Editor/PropertyEditor/Private/CustomChildBuilder.cpp
+++ b/Engine/Source/Editor/PropertyEditor/Private/CustomChildBuilder.cpp
@@ -91,7 +91,7 @@ TSharedRef<SWidget> FCustomChildrenBuilder::GenerateStructValueWidget( TSharedRe
 	UStructProperty* StructProperty = CastChecked<UStructProperty>( StructPropertyHandle->GetProperty() );
 
 	FPropertyEditorModule& PropertyEditorModule = FModuleManager::GetModuleChecked<FPropertyEditorModule>("PropertyEditor");
-	FOnGetStructCustomizationInstance StructLayoutInstanceDelegate = PropertyEditorModule.GetStructCustomizaton(StructProperty->Struct->GetFName());
+	FOnGetStructCustomizationInstance StructLayoutInstanceDelegate = PropertyEditorModule.GetStructCustomizaton(StructProperty->Struct);
 	if (StructLayoutInstanceDelegate.IsBound())
 	{
 		TSharedRef<IStructCustomization> CustomStructInterface = StructLayoutInstanceDelegate.Execute();
diff --git a/Engine/Source/Editor/PropertyEditor/Private/DetailPropertyRow.cpp b/Engine/Source/Editor/PropertyEditor/Private/DetailPropertyRow.cpp
--- a/Engine/Source/Editor/PropertyEditor/Private/DetailPropertyRow.cpp
+++ b/Engine/Source/Editor/PropertyEditor/Private/DetailPropertyRow.cpp
@@ -35,7 +35,7 @@ FDetailPropertyRow::FDetailPropertyRow(TSharedPtr<FPropertyNode> InPropertyNode,
 		if (!PropertyEditorHelpers::IsStaticArray(*PropertyNodeRef) && StructProperty && StructProperty->Struct)
 		{
 			FPropertyEditorModule& PropertyEditorModule = FModuleManager::GetModuleChecked<FPropertyEditorModule>("PropertyEditor");
-			FOnGetStructCustomizationInstance StructLayoutInstanceDelegate = PropertyEditorModule.GetStructCustomizaton(StructProperty->Struct->GetFName());
+			FOnGetStructCustomizationInstance StructLayoutInstanceDelegate = PropertyEditorModule.GetStructCustomizaton(StructProperty->Struct);
 			if (StructLayoutInstanceDelegate.IsBound())
 			{
 				if (PropertyHandle->IsValidHandle())
diff --git a/Engine/Source/Editor/PropertyEditor/Private/PropertyEditorModule.cpp b/Engine/Source/Editor/PropertyEditor/Private/PropertyEditorModule.cpp
--- a/Engine/Source/Editor/PropertyEditor/Private/PropertyEditorModule.cpp
+++ b/Engine/Source/Editor/PropertyEditor/Private/PropertyEditorModule.cpp
@@ -473,4 +473,20 @@ void FPropertyEditorModule::RemoveDeletedObjects( TArray<UObject*>& DeletedObjec
 	}
 }
 
+bool FPropertyEditorModule::IsCustomizedStruct(const UStruct* Struct) const
+{
+	if (Struct && !Struct->IsA<UUserDefinedStruct>())
+	{
+		return StructTypeToLayoutMap.Contains(Struct->GetFName());
+	}
+	return false;
+}
 
+FOnGetStructCustomizationInstance FPropertyEditorModule::GetStructCustomizaton(const UStruct* Struct)
+{
+	if (Struct && !Struct->IsA<UUserDefinedStruct>())
+	{
+		return StructTypeToLayoutMap.FindRef(Struct->GetFName());
+	}
+	return FOnGetStructCustomizationInstance();
+}
\ No newline at end of file
diff --git a/Engine/Source/Editor/PropertyEditor/Private/PropertyHandleImpl.cpp b/Engine/Source/Editor/PropertyEditor/Private/PropertyHandleImpl.cpp
--- a/Engine/Source/Editor/PropertyEditor/Private/PropertyHandleImpl.cpp
+++ b/Engine/Source/Editor/PropertyEditor/Private/PropertyHandleImpl.cpp
@@ -1070,7 +1070,7 @@ void FPropertyValueImpl::AddChild()
 
 					if (const UStructProperty* StructProperty = Cast<const UStructProperty>(Array->Inner))
 					{
-						if(!FStructureEditorUtils::Fill_MakeStructureDefaultValue(Cast<const UBlueprintGeneratedStruct>(StructProperty->Struct), ArrayHelper.GetRawPtr(ArrayIndex)))
+						if(!FStructureEditorUtils::Fill_MakeStructureDefaultValue(Cast<const UUserDefinedStruct>(StructProperty->Struct), ArrayHelper.GetRawPtr(ArrayIndex)))
 						{
 							UE_LOG(LogPropertyNode, Warning, TEXT("MakeStructureDefaultValue parsing error. Property: %s "), *StructProperty->GetName());
 						}
diff --git a/Engine/Source/Editor/PropertyEditor/Private/SDetailsView.cpp b/Engine/Source/Editor/PropertyEditor/Private/SDetailsView.cpp
--- a/Engine/Source/Editor/PropertyEditor/Private/SDetailsView.cpp
+++ b/Engine/Source/Editor/PropertyEditor/Private/SDetailsView.cpp
@@ -396,7 +396,7 @@ static bool IsVisibleStandaloneProperty( const FPropertyNode& PropertyNode, cons
 	return bIsVisibleStandalone;
 }
 
-void SDetailsView::UpdatePropertyMapRecursive( FPropertyNode& InNode, FDetailLayoutBuilderImpl& InDetailLayout, FName CurCategory, FObjectPropertyNode* CurObjectNode, const FCustomStructLayoutMap& StructLayoutMap )
+void SDetailsView::UpdatePropertyMapRecursive( FPropertyNode& InNode, FDetailLayoutBuilderImpl& InDetailLayout, FName CurCategory, FObjectPropertyNode* CurObjectNode )
 {
 	UProperty* ParentProperty = InNode.GetProperty();
 	UStructProperty* ParentStructProp = Cast<UStructProperty>(ParentProperty);
@@ -414,12 +414,12 @@ void SDetailsView::UpdatePropertyMapRecursive( FPropertyNode& InNode, FDetailLay
 			{
 				// Currently object property nodes do not provide any useful information other than being a container for its children.  We do not draw anything for them.
 				// When we encounter object property nodes, add their children instead of adding them to the tree.
-				UpdatePropertyMapRecursive( ChildNode, InDetailLayout, CurCategory, ObjNode, StructLayoutMap );
+				UpdatePropertyMapRecursive( ChildNode, InDetailLayout, CurCategory, ObjNode );
 			}
 			else if( CategoryNode )
 			{
 				// For category nodes, we just set the current category and recurse through the children
-				UpdatePropertyMapRecursive( ChildNode, InDetailLayout, CategoryNode->GetCategoryName(), CurObjectNode, StructLayoutMap );
+				UpdatePropertyMapRecursive( ChildNode, InDetailLayout, CategoryNode->GetCategoryName(), CurObjectNode );
 			}
 			else
 			{
@@ -435,14 +435,21 @@ void SDetailsView::UpdatePropertyMapRecursive( FPropertyNode& InNode, FDetailLay
 
 				bool bIsChildOfCustomizedStruct = false;
 				bool bIsCustomizedStruct = false;
-				if( StructProperty && StructProperty->Struct )
-				{
-					bIsCustomizedStruct = StructLayoutMap.Contains( StructProperty->Struct->GetFName() );
-				}
 
-				if( ParentStructProp && ParentStructProp->Struct )
+				const UStruct* Struct = StructProperty ? StructProperty->Struct : NULL;
+				const UStruct* ParentStruct = ParentStructProp ? ParentStructProp->Struct : NULL;
+				if (Struct || ParentStruct)
 				{
-					bIsChildOfCustomizedStruct = StructLayoutMap.Contains( ParentStructProp->Struct->GetFName() );
+					FPropertyEditorModule& ParentPlugin = FModuleManager::GetModuleChecked<FPropertyEditorModule>("PropertyEditor");
+					if (Struct)
+					{
+						bIsCustomizedStruct = ParentPlugin.IsCustomizedStruct(Struct);
+					}
+
+					if (ParentStruct)
+					{
+						bIsChildOfCustomizedStruct = ParentPlugin.IsCustomizedStruct(ParentStruct);
+					}
 				}
 
 				// Whether or not to push out struct properties to their own categories or show them inside an expandable struct 
@@ -536,7 +543,7 @@ void SDetailsView::UpdatePropertyMapRecursive( FPropertyNode& InNode, FDetailLay
 				if( bRecurseIntoChildren )
 				{
 					// Built in struct properties or children of arras 
-					UpdatePropertyMapRecursive( ChildNode, InDetailLayout, CurCategory, CurObjectNode, StructLayoutMap );
+					UpdatePropertyMapRecursive( ChildNode, InDetailLayout, CurCategory, CurObjectNode );
 				}
 			}		
 		}
@@ -557,15 +564,10 @@ void SDetailsView::UpdatePropertyMap()
 	RootTreeNodes.Empty();
 
 	DetailLayout = MakeShareable( new FDetailLayoutBuilderImpl( ClassToPropertyMap, PropertyUtilities.ToSharedRef(), SharedThis( this ) ) );
-	
-	FPropertyEditorModule& ParentPlugin = FModuleManager::GetModuleChecked<FPropertyEditorModule>("PropertyEditor");
-
-	// Get the registered classes that customize details
-	FCustomStructLayoutMap& GlobalStructLayoutMap = ParentPlugin.StructTypeToLayoutMap;
 
 	// Currently object property nodes do not provide any useful information other than being a container for its children.  We do not draw anything for them.
 	// When we encounter object property nodes, add their children instead of adding them to the tree.
-	UpdatePropertyMapRecursive( *RootPropertyNode, *DetailLayout, NAME_None, RootPropertyNode.Get(), GlobalStructLayoutMap );
+	UpdatePropertyMapRecursive( *RootPropertyNode, *DetailLayout, NAME_None, RootPropertyNode.Get() );
 
 	// Ask for custom detail layouts
 	QueryCustomDetailLayout( *DetailLayout );
diff --git a/Engine/Source/Editor/PropertyEditor/Private/SDetailsView.h b/Engine/Source/Editor/PropertyEditor/Private/SDetailsView.h
--- a/Engine/Source/Editor/PropertyEditor/Private/SDetailsView.h
+++ b/Engine/Source/Editor/PropertyEditor/Private/SDetailsView.h
@@ -350,7 +350,7 @@ private:
 	 * @param The detail layout builder that will be used for customization of this property map
 	 * @param CurCategory The current category name
 	 */
-	void UpdatePropertyMapRecursive( FPropertyNode& InNode, FDetailLayoutBuilderImpl& DetailLayout, FName CurCategory, FObjectPropertyNode* CurObjectNode, const FCustomStructLayoutMap& StructLayoutMap );
+	void UpdatePropertyMapRecursive( FPropertyNode& InNode, FDetailLayoutBuilderImpl& DetailLayout, FName CurCategory, FObjectPropertyNode* CurObjectNode );
 
 	/** Called before during SetObjectArray before we change the objects being observed */
 	void PreSetObject();
diff --git a/Engine/Source/Editor/PropertyEditor/Public/PropertyEditorModule.h b/Engine/Source/Editor/PropertyEditor/Public/PropertyEditorModule.h
--- a/Engine/Source/Editor/PropertyEditor/Public/PropertyEditorModule.h
+++ b/Engine/Source/Editor/PropertyEditor/Public/PropertyEditorModule.h
@@ -317,10 +317,8 @@ public:
 	virtual TSharedRef< FAssetEditorToolkit > CreatePropertyEditorToolkit( const EToolkitMode::Type Mode, const TSharedPtr< IToolkitHost >& InitToolkitHost, const TArray< UObject* >& ObjectsToEdit );
 	virtual TSharedRef< FAssetEditorToolkit > CreatePropertyEditorToolkit( const EToolkitMode::Type Mode, const TSharedPtr< IToolkitHost >& InitToolkitHost, const TArray< TWeakObjectPtr< UObject > >& ObjectsToEdit );
 
-	FOnGetStructCustomizationInstance GetStructCustomizaton( FName StructTypeName )
-	{
-		return StructTypeToLayoutMap.FindRef( StructTypeName );
-	}
+	FOnGetStructCustomizationInstance GetStructCustomizaton(const UStruct* Struct);
+	bool IsCustomizedStruct(const UStruct* Struct) const;
 
 private:
 
diff --git a/Engine/Source/Editor/ReferenceViewer/Private/ReferenceViewer.cpp b/Engine/Source/Editor/ReferenceViewer/Private/ReferenceViewer.cpp
--- a/Engine/Source/Editor/ReferenceViewer/Private/ReferenceViewer.cpp
+++ b/Engine/Source/Editor/ReferenceViewer/Private/ReferenceViewer.cpp
@@ -2,7 +2,7 @@
 
 #include "ReferenceViewerPrivatePCH.h"
 #include "ReferenceViewer.h"
-#include "ReferenceViewer.generated.inl"
+
 #include "EdGraphUtilities.h"
 
 #define LOCTEXT_NAMESPACE "ReferenceViewer"
diff --git a/Engine/Source/Editor/StatsViewer/Private/StatsViewerModule.cpp b/Engine/Source/Editor/StatsViewer/Private/StatsViewerModule.cpp
--- a/Engine/Source/Editor/StatsViewer/Private/StatsViewerModule.cpp
+++ b/Engine/Source/Editor/StatsViewer/Private/StatsViewerModule.cpp
@@ -7,7 +7,7 @@
 #include "StaticMeshLightingInfoStatsPage.h"
 #include "TextureStatsPage.h"
 #include "ObjectHyperlinkColumn.h"
-#include "StatsViewer.generated.inl"
+
 
 #define LOCTEXT_NAMESPACE "Editor.StatsViewer"
 
diff --git a/Engine/Source/Editor/SwarmInterface/Private/SwarmInterfaceModule.cpp b/Engine/Source/Editor/SwarmInterface/Private/SwarmInterfaceModule.cpp
--- a/Engine/Source/Editor/SwarmInterface/Private/SwarmInterfaceModule.cpp
+++ b/Engine/Source/Editor/SwarmInterface/Private/SwarmInterfaceModule.cpp
@@ -4,7 +4,7 @@
 
 #include "ModuleManager.h"
 
-#include "SwarmInterface.generated.inl"
+#include "SwarmMessages.h"
 
 class FSwarmInterfaceModule : public IModuleInterface
 {
diff --git a/Engine/Source/Editor/TextureEditor/Private/TextureEditorClasses.cpp b/Engine/Source/Editor/TextureEditor/Private/TextureEditorClasses.cpp
--- a/Engine/Source/Editor/TextureEditor/Private/TextureEditorClasses.cpp
+++ b/Engine/Source/Editor/TextureEditor/Private/TextureEditorClasses.cpp
@@ -5,7 +5,7 @@
 =============================================================================*/
 
 #include "TextureEditorPrivatePCH.h"
-#include "TextureEditor.generated.inl"
+
 
 
 UTextureEditorSettings::UTextureEditorSettings( const class FPostConstructInitializeProperties& PCIP )
diff --git a/Engine/Source/Editor/TranslationEditor/Private/TranslationUnit.cpp b/Engine/Source/Editor/TranslationEditor/Private/TranslationUnit.cpp
--- a/Engine/Source/Editor/TranslationEditor/Private/TranslationUnit.cpp
+++ b/Engine/Source/Editor/TranslationEditor/Private/TranslationUnit.cpp
@@ -1,6 +1,5 @@
 // Copyright 1998-2014 Epic Games, Inc. All Rights Reserved.
 #include "TranslationEditorPrivatePCH.h"
-#include "TranslationEditor.generated.inl"
 
 UTranslationUnit::UTranslationUnit( const class FPostConstructInitializeProperties& PCIP )
 	: Super(PCIP)
diff --git a/Engine/Source/Editor/UMGEditor/Private/UMGEditorModule.cpp b/Engine/Source/Editor/UMGEditor/Private/UMGEditorModule.cpp
--- a/Engine/Source/Editor/UMGEditor/Private/UMGEditorModule.cpp
+++ b/Engine/Source/Editor/UMGEditor/Private/UMGEditorModule.cpp
@@ -11,8 +11,6 @@
 #include "KismetCompilerModule.h"
 #include "WidgetBlueprintCompiler.h"
 
-#include "UMGEditor.generated.inl"
-
 const FName UMGEditorAppIdentifier = FName(TEXT("UMGEditorApp"));
 
 class FUMGEditorModule : public IUMGEditorModule
diff --git a/Engine/Source/Editor/UnrealEd/Classes/Editor/Transactor.h b/Engine/Source/Editor/UnrealEd/Classes/Editor/Transactor.h
--- a/Engine/Source/Editor/UnrealEd/Classes/Editor/Transactor.h
+++ b/Engine/Source/Editor/UnrealEd/Classes/Editor/Transactor.h
@@ -333,7 +333,6 @@ public:
 	/** Serializes a reference to a transaction in a given archive. */
 	friend FArchive& operator<<( FArchive& Ar, FTransaction& T )
 	{
-		check( Ar.IsAllowingReferenceElimination() == false );
 		return Ar << T.Records << T.Title << T.ObjectMap << T.Context << T.PrimaryObject;
 	}
 
diff --git a/Engine/Source/Editor/UnrealEd/Classes/Factories/StructureFactory.h b/Engine/Source/Editor/UnrealEd/Classes/Factories/StructureFactory.h
new file mode 100644
--- /dev/null
+++ b/Engine/Source/Editor/UnrealEd/Classes/Factories/StructureFactory.h
@@ -0,0 +1,18 @@
+// Copyright 1998-2013 Epic Games, Inc. All Rights Reserved.
+
+#pragma once
+
+#include "StructureFactory.generated.h"
+
+UCLASS(hidecategories=Object, MinimalAPI)
+class UStructureFactory : public UFactory
+{
+	GENERATED_UCLASS_BODY()
+
+	// Begin UFactory Interface
+	virtual UObject* FactoryCreateNew(UClass* Class,UObject* InParent,FName Name,EObjectFlags Flags,UObject* Context,FFeedbackContext* Warn) OVERRIDE;
+	// Begin UFactory Interface
+};
+
+
+
diff --git a/Engine/Source/Editor/UnrealEd/Classes/UserDefinedStructure/UserDefinedStructEditorData.h b/Engine/Source/Editor/UnrealEd/Classes/UserDefinedStructure/UserDefinedStructEditorData.h
new file mode 100644
--- /dev/null
+++ b/Engine/Source/Editor/UnrealEd/Classes/UserDefinedStructure/UserDefinedStructEditorData.h
@@ -0,0 +1,70 @@
+// Copyright 1998-2014 Epic Games, Inc. All Rights Reserved.
+
+#pragma once
+
+#include "UserDefinedStructEditorData.generated.h"
+
+USTRUCT()
+struct FStructVariableDescription
+{
+	GENERATED_USTRUCT_BODY()
+
+	UPROPERTY()
+	FName VarName;
+
+	UPROPERTY()
+	FGuid VarGuid;
+
+	UPROPERTY()
+	FString FriendlyName;
+
+	UPROPERTY()
+	FString DefaultValue;
+
+	// TYPE DATA
+	UPROPERTY()
+	FString Category;
+
+	UPROPERTY()
+	FString SubCategory;
+
+	UPROPERTY()
+	TAssetPtr<UObject> SubCategoryObject;
+
+	UPROPERTY()
+	bool bIsArray;
+
+	UPROPERTY(Transient)
+	bool bInvalidMember;
+
+	// CurrentDefaultValue stores the actual default value, after the DefaultValue was changed, and before the struct was recompiled
+	UPROPERTY()
+	FString CurrentDefaultValue;
+
+	UNREALED_API bool SetPinType(const struct FEdGraphPinType& VarType);
+
+	UNREALED_API FEdGraphPinType ToPinType() const;
+};
+
+UCLASS()
+class UNREALED_API UUserDefinedStructEditorData : public UObject
+{
+	GENERATED_UCLASS_BODY()
+
+private:
+	// the property is used to generate an uniqe name id for member variable
+	UPROPERTY(NonTransactional) 
+	uint32 UniqueNameId;
+
+public:
+	UPROPERTY()
+	TArray<FStructVariableDescription> VariablesDescriptions;
+
+	// UObject interface.
+	virtual void PostEditUndo() OVERRIDE;
+	virtual void PostLoadSubobjects(struct FObjectInstancingGraph* OuterInstanceGraph) OVERRIDE;
+	// End of UObject interface.
+
+	uint32 GenerateUniqueNameIdForMemberVariable();
+	class UUserDefinedStruct* GetOwnerStruct() const;
+};
\ No newline at end of file
diff --git a/Engine/Source/Editor/UnrealEd/Private/EditorConstraints.cpp b/Engine/Source/Editor/UnrealEd/Private/EditorConstraints.cpp
--- a/Engine/Source/Editor/UnrealEd/Private/EditorConstraints.cpp
+++ b/Engine/Source/Editor/UnrealEd/Private/EditorConstraints.cpp
@@ -9,8 +9,14 @@
 
 float UEditorEngine::GetGridSize()
 {
-	TArray<float> PosGridSizes = GetCurrentPositionGridArray();
-	return PosGridSizes[GetDefault<ULevelEditorViewportSettings>()->CurrentPosGridSize ];
+	const TArray<float> PosGridSizes = GetCurrentPositionGridArray();
+	const int32 CurrentPosGridSize = GetDefault<ULevelEditorViewportSettings>()->CurrentPosGridSize;
+	float PosVal = 0.0001f;
+	if ( PosGridSizes.IsValidIndex(CurrentPosGridSize) )
+	{
+		PosVal = PosGridSizes[CurrentPosGridSize];
+	}
+	return PosVal;
 }
 
 bool UEditorEngine::IsGridSizePowerOfTwo() const
@@ -57,8 +63,13 @@ const TArray<float>& UEditorEngine::GetCurrentPositionGridArray() const
 
 FRotator UEditorEngine::GetRotGridSize()
 {
-	TArray<float> RotGridSizes = GetCurrentRotationGridArray();
-	float RotVal = RotGridSizes[GetDefault<ULevelEditorViewportSettings>()->CurrentRotGridSize ];
+	const TArray<float> RotGridSizes = GetCurrentRotationGridArray();
+	const int32 CurrentRotGridSize = GetDefault<ULevelEditorViewportSettings>()->CurrentRotGridSize;
+	float RotVal = 0.0001f;
+	if ( RotGridSizes.IsValidIndex(CurrentRotGridSize) )
+	{
+		RotVal = RotGridSizes[CurrentRotGridSize];
+	}
 	return FRotator(RotVal, RotVal, RotVal);
 }
 
@@ -66,7 +77,7 @@ void UEditorEngine::SetRotGridSize( int32 InIndex, ERotationGridMode InGridMode
 {
 	FinishAllSnaps();
 
-	TArray<float> RotGridSizes = GetCurrentRotationGridArray();
+	const TArray<float> RotGridSizes = GetCurrentRotationGridArray();
 
 	ULevelEditorViewportSettings* ViewportSettings = GetMutableDefault<ULevelEditorViewportSettings>();
 	ViewportSettings->CurrentRotGridMode = InGridMode;
@@ -101,8 +112,12 @@ const TArray<float>& UEditorEngine::GetCurrentRotationGridArray() const
 float UEditorEngine::GetScaleGridSize()
 {
 	const ULevelEditorViewportSettings* ViewportSettings = GetDefault<ULevelEditorViewportSettings>();
-
-	return ViewportSettings->ScalingGridSizes[ViewportSettings->CurrentScalingGridSize];
+	float ScaleVal = 0.0001f;
+	if ( ViewportSettings->ScalingGridSizes.IsValidIndex(ViewportSettings->CurrentScalingGridSize) )
+	{
+		ScaleVal = ViewportSettings->ScalingGridSizes[ViewportSettings->CurrentScalingGridSize];
+	}
+	return ScaleVal;
 }
 
 void UEditorEngine::SetScaleGridSize( int32 InIndex )
@@ -120,7 +135,6 @@ void UEditorEngine::SetScaleGridSize( int32 InIndex )
 float UEditorEngine::GetGridInterval()
 {
 	const ULevelEditorViewportSettings* ViewportSettings = GetDefault<ULevelEditorViewportSettings>();
-
 	const TArray<float>& GridIntervals = ViewportSettings->bUsePowerOf2SnapSize ? ViewportSettings->Pow2GridIntervals : ViewportSettings->DecimalGridIntervals;
 
 	int32 CurrentIndex = (ViewportSettings->CurrentPosGridSize < GridIntervals.Num()) ? ViewportSettings->CurrentPosGridSize : GridIntervals.Num() - 1;
diff --git a/Engine/Source/Editor/UnrealEd/Private/EditorTransaction.cpp b/Engine/Source/Editor/UnrealEd/Private/EditorTransaction.cpp
--- a/Engine/Source/Editor/UnrealEd/Private/EditorTransaction.cpp
+++ b/Engine/Source/Editor/UnrealEd/Private/EditorTransaction.cpp
@@ -314,16 +314,11 @@ void UTransBuffer::Serialize( FArchive& Ar )
 
 	CheckState();
 
-	// Handle garbage collection.
 	Super::Serialize( Ar );
 
-	// We cannot support undoing across GC if we allow it to eliminate references so we need
-	// to suppress it.
 	if ( IsObjectSerializationEnabled() || !Ar.IsObjectReferenceCollector() )
 	{
-		Ar.AllowEliminatingReferences( false );
 		Ar << UndoBuffer;
-		Ar.AllowEliminatingReferences( true );
 	}
 	Ar << ResetReason << UndoCount << ActiveCount;
 
diff --git a/Engine/Source/Editor/UnrealEd/Private/Factories/EditorFactories.cpp b/Engine/Source/Editor/UnrealEd/Private/Factories/EditorFactories.cpp
--- a/Engine/Source/Editor/UnrealEd/Private/Factories/EditorFactories.cpp
+++ b/Engine/Source/Editor/UnrealEd/Private/Factories/EditorFactories.cpp
@@ -37,6 +37,7 @@
 #include "SClassPickerDialog.h"
 #include "MessageLog.h"
 #include "EnumEditorUtils.h"
+#include "StructureEditorUtils.h"
 
 #if PLATFORM_WINDOWS
 // Needed for DDS support.
@@ -7057,6 +7058,23 @@ UObject* UEnumFactory::FactoryCreateNew(UClass* Class, UObject* InParent, FName
 	return FEnumEditorUtils::CreateUserDefinedEnum(InParent, Name, Flags);
 }
 
+/*------------------------------------------------------------------------------
+	UStructureFactory implementation.
+------------------------------------------------------------------------------*/
+UStructureFactory::UStructureFactory(const class FPostConstructInitializeProperties& PCIP)
+	: Super(PCIP)
+{
+	SupportedClass = UUserDefinedStruct::StaticClass();
+	bCreateNew = FStructureEditorUtils::UserDefinedStructEnabled();
+	bEditAfterNew = true;
+}
+
+UObject* UStructureFactory::FactoryCreateNew(UClass* Class, UObject* InParent, FName Name, EObjectFlags Flags, UObject* Context, FFeedbackContext* Warn)
+{
+	ensure(UUserDefinedStruct::StaticClass() == Class);
+	return FStructureEditorUtils::CreateUserDefinedStruct(InParent, Name, Flags);
+}
+
 /*-----------------------------------------------------------------------------
 	UForceFeedbackEffectFactory implementation.
 -----------------------------------------------------------------------------*/
diff --git a/Engine/Source/Editor/UnrealEd/Private/Kismet2/BlueprintEditorUtils.cpp b/Engine/Source/Editor/UnrealEd/Private/Kismet2/BlueprintEditorUtils.cpp
--- a/Engine/Source/Editor/UnrealEd/Private/Kismet2/BlueprintEditorUtils.cpp
+++ b/Engine/Source/Editor/UnrealEd/Private/Kismet2/BlueprintEditorUtils.cpp
@@ -21,6 +21,7 @@
 #include "DefaultValueHelper.h"
 #include "ObjectEditorUtils.h"
 #include "ToolkitManager.h"
+#include "Runtime/Engine/Classes/Engine/UserDefinedStruct.h"
 
 #define LOCTEXT_NAMESPACE "Blueprint"
 
@@ -577,6 +578,18 @@ struct FRegenerationHelper
 		}
 	}
 
+	static UBlueprint* GetGeneratingBlueprint(const UObject* Obj)
+	{
+		const UBlueprintGeneratedClass* BPGC = NULL;
+		while (!BPGC && Obj)
+		{
+			BPGC = Cast<const UBlueprintGeneratedClass>(Obj);
+			Obj = Obj->GetOuter();
+		}
+
+		return UBlueprint::GetBlueprintFromClass(BPGC);
+	}
+
 	static void ProcessHierarchy(UStruct* Struct, TSet<UStruct*>& Dependencies)
 	{
 		if (Struct)
@@ -589,7 +602,7 @@ struct FRegenerationHelper
 
 				const UBlueprint* BP = GetGeneratingBlueprint(Struct);
 				const bool bProcessBPGClass = BP && !BP->bHasBeenRegenerated;
-				const bool bProcessUserDefinedStruct = false; //TODO: Struct->IsA<UUserDefinedStruct>();
+				const bool bProcessUserDefinedStruct = Struct->IsA<UUserDefinedStruct>();
 				if (bProcessBPGClass || bProcessUserDefinedStruct)
 				{
 					PreloadAndLinkIfNecessary(Struct);
@@ -598,18 +611,6 @@ struct FRegenerationHelper
 		}
 	}
 
-	static UBlueprint* GetGeneratingBlueprint(const UObject* Obj)
-	{
-		const UBlueprintGeneratedClass* BPGC = NULL;
-		while (!BPGC && Obj)
-		{
-			BPGC = Cast<const UBlueprintGeneratedClass>(Obj);
-			Obj = Obj->GetOuter();
-		}
-
-		return UBlueprint::GetBlueprintFromClass(BPGC);
-	}
-
 	static void PreloadMacroSources(TSet<UBlueprint*>& MacroSources)
 	{
 		for (auto BP : MacroSources)
@@ -666,8 +667,7 @@ struct FRegenerationHelper
 					}
 
 					LocalDependentStructures.Empty(LocalDependentStructures.Max());
-					//TODO: uncomment when UDStructures are ready
-					//Node->HasExternalUserDefinedStructDependencies(&LocalDependentStructures);
+					Node->HasExternalUserDefinedStructDependencies(&LocalDependentStructures);
 					for (auto Struct : LocalDependentStructures)
 					{
 						ProcessHierarchy(Struct, Dependencies);
@@ -692,27 +692,6 @@ struct FRegenerationHelper
 	}
 };
 
-void FBlueprintEditorUtils::RefreshInputDelegatePins(UBlueprint* Blueprint)
-{
-	TArray<UEdGraph*> Graphs;
-	Blueprint->GetAllGraphs(Graphs);
-	for (auto Graph : Graphs)
-	{
-		if (!FBlueprintEditorUtils::IsGraphIntermediate(Graph))
-		{
-			TArray<UK2Node_BaseMCDelegate*> Nodes;
-			Graph->GetNodesOfClass(Nodes);
-			for (auto Node : Nodes)
-			{
-				if (auto DelegatePin = Node->GetDelegatePin())
-				{
-					DelegatePin->PinType.PinSubCategoryObject = Node->GetDelegateSignature();
-				}
-			}
-		}
-	}
-}
-
 struct FEditoronlyBlueprintHelper
 {
 	static bool IsUnwantedType(const FEdGraphPinType& Type)
@@ -1000,7 +979,8 @@ UClass* FBlueprintEditorUtils::RegenerateBlueprintClass(UBlueprint* Blueprint, U
 
 		// Make sure all used external classes/functions/structures/macros/etc are loaded and linked
 		FRegenerationHelper::LinkExternalDependencies(Blueprint);
-		FBlueprintEditorUtils::RefreshInputDelegatePins(Blueprint);
+
+		FKismetEditorUtilities::GenerateBlueprintSkeleton(Blueprint);
 
 		struct FReplaceBlueprintWithClassHelper
 		{
@@ -1363,8 +1343,6 @@ void FBlueprintEditorUtils::PostDuplicateBlueprint(UBlueprint* Blueprint)
 			UEditorEngine::CopyPropertiesForUnrelatedObjects(OldCDO, NewCDO);
 		}
 
-		Blueprint->UserDefinedStructures.Empty();
-
 		// And compile again to make sure they go into the generated class, get cleaned up, etc...
 		FKismetEditorUtilities::CompileBlueprint(Blueprint, false, true);
 
@@ -1395,7 +1373,7 @@ void FBlueprintEditorUtils::UpdateDelegatesInBlueprint(UBlueprint* Blueprint)
 			Graph->GetNodesOfClass(CreateDelegateNodes);
 			for (auto NodeIt = CreateDelegateNodes.CreateIterator(); NodeIt; ++NodeIt)
 			{
-				(*NodeIt)->HandleAnyChangeInner();
+				(*NodeIt)->HandleAnyChangeWithoutNotifying();
 			}
 
 			TArray<UK2Node_Event*> EventNodes;
@@ -2883,6 +2861,18 @@ void FBlueprintEditorUtils::GetHiddenPinsForFunction(UBlueprint const* CallingCo
 	}
 }
 
+bool FBlueprintEditorUtils::IsPinTypeValid(const FEdGraphPinType& Type)
+{
+	if (const auto UDStruct = Cast<const UUserDefinedStruct>(Type.PinSubCategoryObject.Get()))
+	{
+		if (EUserDefinedStructureStatus::UDSS_UpToDate != UDStruct->Status)
+		{
+			return false;
+		}
+	}
+	return true;
+}
+
 // Gets the visible class variable list.  This includes both variables introduced here and in all superclasses.
 void FBlueprintEditorUtils::GetClassVariableList(const UBlueprint* Blueprint, TArray<FName>& VisibleVariables, bool bIncludePrivateVars) 
 {
@@ -5587,6 +5577,95 @@ bool FBlueprintEditorUtils::IsPaletteActionReadOnly(TSharedPtr<FEdGraphSchemaAct
 	return bIsReadOnly;
 }
 
+struct FUberGraphHelper
+{
+	static void GetAll(const UBlueprint* Blueprint, TArray<UEdGraph*>& OutGraphs)
+	{
+		for (const auto UberGraph : Blueprint->UbergraphPages)
+		{
+			OutGraphs.Add(UberGraph);
+			UberGraph->GetAllChildrenGraphs(OutGraphs);
+		}
+	}
+};
+
+FName FBlueprintEditorUtils::GetFunctionNameFromClassByGuid(const UClass* InClass, const FGuid FunctionGuid)
+{
+	TArray<UBlueprint*> Blueprints;
+	UBlueprint::GetBlueprintHierarchyFromClass(InClass, Blueprints);
+
+	for (int32 BPIndex = 0; BPIndex < Blueprints.Num(); ++BPIndex)
+	{
+		UBlueprint* Blueprint = Blueprints[BPIndex];
+		for (int32 FunctionIndex = 0; FunctionIndex < Blueprint->FunctionGraphs.Num(); ++FunctionIndex)
+		{
+			UEdGraph* FunctionGraph = Blueprint->FunctionGraphs[FunctionIndex];
+			if (FunctionGraph && FunctionGraph->GraphGuid == FunctionGuid)
+			{
+				return FunctionGraph->GetFName();
+			}
+		}
+
+		//FUNCTIONS BASED ON CUSTOM EVENTS:
+		TArray<UEdGraph*> UberGraphs;
+		FUberGraphHelper::GetAll(Blueprint, UberGraphs);
+		for (const auto UberGraph : UberGraphs)
+		{
+			TArray<UK2Node_CustomEvent*> CustomEvents;
+			UberGraph->GetNodesOfClass(CustomEvents);
+			for (const auto CustomEvent : CustomEvents)
+			{
+				if (!CustomEvent->bOverrideFunction && (CustomEvent->NodeGuid == FunctionGuid))
+				{
+					ensure(CustomEvent->CustomFunctionName != NAME_None);
+					return CustomEvent->CustomFunctionName;
+				}
+			}
+		}
+	}
+
+	return NAME_None;
+}
+
+bool FBlueprintEditorUtils::GetFunctionGuidFromClassByFieldName(const UClass* InClass, const FName FunctionName, FGuid& FunctionGuid)
+{
+	TArray<UBlueprint*> Blueprints;
+	UBlueprint::GetBlueprintHierarchyFromClass(InClass, Blueprints);
+
+	for (int32 BPIndex = 0; BPIndex < Blueprints.Num(); ++BPIndex)
+	{
+		UBlueprint* Blueprint = Blueprints[BPIndex];
+		for (int32 FunctionIndex = 0; FunctionIndex < Blueprint->FunctionGraphs.Num(); ++FunctionIndex)
+		{
+			UEdGraph* FunctionGraph = Blueprint->FunctionGraphs[FunctionIndex];
+			if (FunctionGraph && FunctionGraph->GetFName() == FunctionName)
+			{
+				FunctionGuid = FunctionGraph->GraphGuid;
+				return true;
+			}
+		}
+
+		TArray<UEdGraph*> UberGraphs;
+		FUberGraphHelper::GetAll(Blueprint, UberGraphs);
+		for (const auto UberGraph : UberGraphs)
+		{
+			TArray<UK2Node_CustomEvent*> CustomEvents;
+			UberGraph->GetNodesOfClass(CustomEvents);
+			for (const auto CustomEvent : CustomEvents)
+			{
+				if (!CustomEvent->bOverrideFunction && (CustomEvent->CustomFunctionName == FunctionName))
+				{
+					ensure(CustomEvent->CustomFunctionName != NAME_None);
+					ensure(CustomEvent->NodeGuid.IsValid());
+					FunctionGuid = CustomEvent->NodeGuid;
+					return true;
+				}
+			}
+		}
+	}
+
+	return false;
+}
 
 #undef LOCTEXT_NAMESPACE
 
diff --git a/Engine/Source/Editor/UnrealEd/Private/Kismet2/EnumEditorUtils.cpp b/Engine/Source/Editor/UnrealEd/Private/Kismet2/EnumEditorUtils.cpp
--- a/Engine/Source/Editor/UnrealEd/Private/Kismet2/EnumEditorUtils.cpp
+++ b/Engine/Source/Editor/UnrealEd/Private/Kismet2/EnumEditorUtils.cpp
@@ -16,7 +16,7 @@ struct FEnumEditorUtilsHelper
 
 //////////////////////////////////////////////////////////////////////////
 // FEnumEditorManager
-FEnumEditorUtils::FEnumEditorManager::FEnumEditorManager()
+FEnumEditorUtils::FEnumEditorManager::FEnumEditorManager() : FListenerManager<UUserDefinedEnum>()
 {}
 
 FEnumEditorUtils::FEnumEditorManager& FEnumEditorUtils::FEnumEditorManager::Get()
@@ -25,27 +25,6 @@ FEnumEditorUtils::FEnumEditorManager& FEnumEditorUtils::FEnumEditorManager::Get(
 	return *EnumEditorManager;
 }
 
-void FEnumEditorUtils::FEnumEditorManager::AddEnumListener(FEnumEditorUtils::INotifyOnEnumChanged* EnumListener)
-{
-	if (EnumListener)
-	{
-		Listeners.Add(EnumListener);
-	}
-}
-
-void FEnumEditorUtils::FEnumEditorManager::RemoveEnumListener(FEnumEditorUtils::INotifyOnEnumChanged* EnumListener)
-{
-	Listeners.Remove(EnumListener);
-}
-
-void FEnumEditorUtils::FEnumEditorManager::OnEnumChanged(const UUserDefinedEnum* Enum)
-{
-	for (auto Iter = Listeners.CreateIterator(); Iter; ++Iter)
-	{
-		(*Iter)->OnEnumChanged(Enum);
-	}
-}
-
 //////////////////////////////////////////////////////////////////////////
 // User defined enumerations
 
@@ -313,7 +292,7 @@ void FEnumEditorUtils::BroadcastChanges(const UUserDefinedEnum* Enum, const TArr
 		(*It)->BroadcastChanged();
 	}
 
-	FEnumEditorManager::Get().OnEnumChanged(Enum);
+	FEnumEditorManager::Get().OnChanged(Enum);
 }
 
 int32 FEnumEditorUtils::ResolveEnumerator(const UEnum* Enum, FArchive& Ar, int32 EnumeratorIndex)
diff --git a/Engine/Source/Editor/UnrealEd/Private/Kismet2/Kismet2.cpp b/Engine/Source/Editor/UnrealEd/Private/Kismet2/Kismet2.cpp
--- a/Engine/Source/Editor/UnrealEd/Private/Kismet2/Kismet2.cpp
+++ b/Engine/Source/Editor/UnrealEd/Private/Kismet2/Kismet2.cpp
@@ -317,14 +317,6 @@ void FKismetEditorUtilities::CompileBlueprint(UBlueprint* BlueprintObj, bool bIs
 	FCompilerResultsLog LocalResults;
 	FCompilerResultsLog& Results = (pResults != NULL) ? *pResults : LocalResults;
 
-	// STRUCTURES RECOMPILE
-	if (FStructureEditorUtils::StructureEditingEnabled())
-	{
-		FKismetCompilerOptions StructCompileOptions;
-		StructCompileOptions.CompileType = EKismetCompileType::StructuresOnly;
-		Compiler.CompileBlueprint(BlueprintObj, StructCompileOptions, Results, NULL);
-	}
-
 	FBlueprintCompileReinstancer ReinstanceHelper(OldClass);
 
 	// Suppress errors/warnings in the log if we're recompiling on load on a build machine
diff --git a/Engine/Source/Editor/UnrealEd/Private/Kismet2/Kismet2NameValidators.cpp b/Engine/Source/Editor/UnrealEd/Private/Kismet2/Kismet2NameValidators.cpp
--- a/Engine/Source/Editor/UnrealEd/Private/Kismet2/Kismet2NameValidators.cpp
+++ b/Engine/Source/Editor/UnrealEd/Private/Kismet2/Kismet2NameValidators.cpp
@@ -92,7 +92,6 @@ FKismetNameValidator::FKismetNameValidator(const class UBlueprint* Blueprint, FN
 	FBlueprintEditorUtils::GetClassVariableList(BlueprintObject, Names, true);
 	FBlueprintEditorUtils::GetAllGraphNames(BlueprintObject, Names);
 	FBlueprintEditorUtils::GetSCSVariableNameList(Blueprint, Names);
-	FStructureEditorUtils::GetAllStructureNames(Blueprint, Names);
 	FBlueprintEditorUtils::GetImplementingBlueprintsFunctionNameList(Blueprint, Names);
 }
 
