#ue4 ttp333624 EDITOR: UDS: Values of text variables aren't propagated to functions

diff --git a/Engine/Source/Editor/KismetCompiler/Private/KismetCompilerVMBackend.cpp b/Engine/Source/Editor/KismetCompiler/Private/KismetCompilerVMBackend.cpp
--- a/Engine/Source/Editor/KismetCompiler/Private/KismetCompilerVMBackend.cpp
+++ b/Engine/Source/Editor/KismetCompiler/Private/KismetCompilerVMBackend.cpp
@@ -476,6 +476,10 @@ public:
 						FBPTerminal NewTerm;
 						NewTerm.bIsLiteral = true;
 						Prop->ExportText_InContainer(0, NewTerm.Name, StructData, StructData, NULL, PPF_None);
+						if (Prop->IsA(UTextProperty::StaticClass()))
+						{
+							NewTerm.TextLiteral = FText::FromString(NewTerm.Name);
+						}
 
 						EmitTermExpr(&NewTerm, Prop);
 					}
