Blueprint Function Library - AutoWorldContext is removed when unnecessary.

diff --git a/Engine/Source/Editor/BlueprintGraph/Classes/K2Node_FunctionEntry.h b/Engine/Source/Editor/BlueprintGraph/Classes/K2Node_FunctionEntry.h
--- a/Engine/Source/Editor/BlueprintGraph/Classes/K2Node_FunctionEntry.h
+++ b/Engine/Source/Editor/BlueprintGraph/Classes/K2Node_FunctionEntry.h
@@ -47,6 +47,8 @@ class UK2Node_FunctionEntry : public UK2Node_FunctionTerminator
 	BLUEPRINTGRAPH_API void RemoveOutputPin(UEdGraphPin* PinToRemove);
 
 	// Returns pin for the automatically added WorldContext parameter (used only by BlueprintFunctionLibrary).
-	BLUEPRINTGRAPH_API UEdGraphPin* GetWorldContextPin() const;
+	BLUEPRINTGRAPH_API UEdGraphPin* GetAutoWorldContextPin() const;
+
+	BLUEPRINTGRAPH_API void RemoveUnnecessaryAutoWorldContext();
 };
 
diff --git a/Engine/Source/Editor/BlueprintGraph/Private/K2Node_CallFunction.cpp b/Engine/Source/Editor/BlueprintGraph/Private/K2Node_CallFunction.cpp
--- a/Engine/Source/Editor/BlueprintGraph/Private/K2Node_CallFunction.cpp
+++ b/Engine/Source/Editor/BlueprintGraph/Private/K2Node_CallFunction.cpp
@@ -1215,7 +1215,7 @@ void UK2Node_CallFunction::ExpandNode(class FKismetCompilerContext& CompilerCont
 			{
 				CompilerContext.MessageLog.Warning(*FString::Printf(*LOCTEXT("WrongEntryPointsNum", "%i entry points found while expanding node @@").ToString(), EntryPoints.Num()), this);
 			}
-			else if (auto BetterSelfPin = EntryPoints[0]->GetWorldContextPin())
+			else if (auto BetterSelfPin = EntryPoints[0]->GetAutoWorldContextPin())
 			{
 				FString const DefaultToSelfMetaValue = Function->GetMetaData(FBlueprintMetadata::MD_DefaultToSelf);
 				FString const WorldContextMetaValue = Function->GetMetaData(FBlueprintMetadata::MD_WorldContext);
diff --git a/Engine/Source/Editor/BlueprintGraph/Private/K2Node_FunctionEntry.cpp b/Engine/Source/Editor/BlueprintGraph/Private/K2Node_FunctionEntry.cpp
--- a/Engine/Source/Editor/BlueprintGraph/Private/K2Node_FunctionEntry.cpp
+++ b/Engine/Source/Editor/BlueprintGraph/Private/K2Node_FunctionEntry.cpp
@@ -161,11 +161,23 @@ void UK2Node_FunctionEntry::AllocateDefaultPins()
 	}
 }
 
-UEdGraphPin* UK2Node_FunctionEntry::GetWorldContextPin() const
+UEdGraphPin* UK2Node_FunctionEntry::GetAutoWorldContextPin() const
 {
 	return FFunctionEntryHelper::RequireWorldContextParameter(this) ? FindPin(FFunctionEntryHelper::GetWorldContextPinName()) : NULL;
 }
 
+void UK2Node_FunctionEntry::RemoveUnnecessaryAutoWorldContext()
+{
+	auto WorldContextPin = GetAutoWorldContextPin();
+	if (WorldContextPin)
+	{
+		if (!WorldContextPin->LinkedTo.Num())
+		{
+			Pins.Remove(WorldContextPin);
+		}
+	}
+}
+
 void UK2Node_FunctionEntry::RemoveOutputPin(UEdGraphPin* PinToRemove)
 {
 	UK2Node_FunctionEntry* OwningSeq = Cast<UK2Node_FunctionEntry>( PinToRemove->GetOwningNode() );
diff --git a/Engine/Source/Editor/KismetCompiler/Private/KismetCompiler.cpp b/Engine/Source/Editor/KismetCompiler/Private/KismetCompiler.cpp
--- a/Engine/Source/Editor/KismetCompiler/Private/KismetCompiler.cpp
+++ b/Engine/Source/Editor/KismetCompiler/Private/KismetCompiler.cpp
@@ -1022,6 +1022,9 @@ void FKismetCompilerContext::PrecompileFunction(FKismetFunctionContext& Context)
 		// Transforms
 		TransformNodes(Context);
 
+		//Now we can safely remove automatically added WorldContext pin from static function.
+		Context.EntryPoint->RemoveUnnecessaryAutoWorldContext();
+
 		// Create the function stub
 		FName NewFunctionName = (Context.EntryPoint->CustomGeneratedFunctionName != NAME_None) ? Context.EntryPoint->CustomGeneratedFunctionName : Context.EntryPoint->SignatureName;
 		if(Context.IsDelegateSignature())
@@ -1451,7 +1454,7 @@ void FKismetCompilerContext::FinishCompilingFunction(FKismetFunctionContext& Con
 		Function->SetMetaData(FBlueprintMetadata::MD_Tooltip, *EntryNode->MetaData.ToolTip);
 	}
 
-	if (auto WorldContextPin = EntryNode->GetWorldContextPin())
+	if (auto WorldContextPin = EntryNode->GetAutoWorldContextPin())
 	{
 		Function->SetMetaData(FBlueprintMetadata::MD_DefaultToSelf, *WorldContextPin->PinName); 
 		Function->SetMetaData(TEXT("HidePin"), *WorldContextPin->PinName);
