[330502] Fix Blueprint compile error message when self pin type is missing or is otherwise unknown.

#change Modified FKismetCompilerContext::ValidateSelfPinsInGraph() to emit an alternative error message to the compiler log when the PinSubCategoryObject is NULL.

diff --git a/Engine/Source/Editor/KismetCompiler/Private/KismetCompiler.cpp b/Engine/Source/Editor/KismetCompiler/Private/KismetCompiler.cpp
--- a/Engine/Source/Editor/KismetCompiler/Private/KismetCompiler.cpp
+++ b/Engine/Source/Editor/KismetCompiler/Private/KismetCompiler.cpp
@@ -969,9 +969,22 @@ void FKismetCompilerContext::ValidateSelfPinsInGraph(const UEdGraph* SourceGraph
 									{
 										PinType = Pin->PinType.PinSubCategoryObject->GetName();
 									}
+									else
+									{
+										PinType = TEXT("");
+									}
+								}
+
+								FString ErrorMsg;
+								if(PinType.IsEmpty())
+								{
+									ErrorMsg = FString::Printf(*LOCTEXT("PinMustHaveConnection_NoType_Error", "'@@' must have a connection").ToString());
+								}
+								else
+								{
+									ErrorMsg = FString::Printf(*LOCTEXT("PinMustHaveConnection_Error", "This blueprint (self) is not a %s, therefore '@@' must have a connection").ToString(), *PinType);	
 								}
 
-								FString ErrorMsg = FString::Printf(*LOCTEXT("PinMustHaveConnection_Error", "This blueprint (self) is not a %s, therefore '@@' must have a connection").ToString(), *PinType);
 								MessageLog.Error(*ErrorMsg, Pin);
 							}
 						}
