Fix for Blueprint function node output pins incorrectly mapped to internal skeleton class types.

- Modified FKismeCompilerContext::ValidatePin() to substitute with a 'self' pin only when the pin type matches the current generated class type.

diff --git a/Engine/Source/Editor/KismetCompiler/Private/KismetCompiler.cpp b/Engine/Source/Editor/KismetCompiler/Private/KismetCompiler.cpp
--- a/Engine/Source/Editor/KismetCompiler/Private/KismetCompiler.cpp
+++ b/Engine/Source/Editor/KismetCompiler/Private/KismetCompiler.cpp
@@ -269,22 +269,11 @@ void FKismetCompilerContext::ValidatePin(const UEdGraphPin* Pin) const
 	{
 		//@todo:  This is modifying the model, but is acceptable to save another prepass on the pins
 		UEdGraphPin* MutablePin = const_cast<UEdGraphPin*>(Pin);
-		if( Blueprint->GeneratedClass)
-		{
-			if (MutablePin->PinType.PinSubCategoryObject.Get() == Blueprint->GeneratedClass)
-			{
-				MutablePin->PinType.PinSubCategory = Schema->PSC_Self;
-				MutablePin->PinType.PinSubCategoryObject = NULL;
-			}
-		}
 
-		if( Blueprint->SkeletonGeneratedClass)
+		if (NewClass && MutablePin->PinType.PinSubCategoryObject.Get() == NewClass)
 		{
-			if (MutablePin->PinType.PinSubCategoryObject.Get() == Blueprint->SkeletonGeneratedClass)
-			{
-				MutablePin->PinType.PinSubCategory = Schema->PSC_Self;
-				MutablePin->PinType.PinSubCategoryObject = NULL;
-			}
+			MutablePin->PinType.PinSubCategory = Schema->PSC_Self;
+			MutablePin->PinType.PinSubCategoryObject = NULL;
 		}
 	}
 
