Better handling improper ArrayProperties. ttp303326

diff --git a/Engine/Source/Editor/KismetCompiler/Private/KismetCompilerMisc.cpp b/Engine/Source/Editor/KismetCompiler/Private/KismetCompilerMisc.cpp
--- a/Engine/Source/Editor/KismetCompiler/Private/KismetCompilerMisc.cpp
+++ b/Engine/Source/Editor/KismetCompiler/Private/KismetCompilerMisc.cpp
@@ -674,11 +674,18 @@ UProperty* FKismetCompilerUtilities::CreatePropertyOnScope(UStruct* Scope, const
 		NewProperty = NewNamedObject<UIntProperty>(PropertyScope, ValidatedPropertyName, ObjectFlags);
 	}
 
-	if( bIsArrayProperty )
+	if (bIsArrayProperty)
 	{
-		// Fix up the array property to have the new type-specific property as its inner, and return the new UArrayProperty
-		NewArrayProperty->Inner = NewProperty;
-		NewProperty = NewArrayProperty;
+		if (NewProperty)
+		{
+			// Fix up the array property to have the new type-specific property as its inner, and return the new UArrayProperty
+			NewArrayProperty->Inner = NewProperty;
+			NewProperty = NewArrayProperty;
+		}
+		else
+		{
+			NewArrayProperty->MarkPendingKill();
+		}
 	}
 
 	return NewProperty;
diff --git a/Engine/Source/Runtime/CoreUObject/Private/UObject/PropertyArray.cpp b/Engine/Source/Runtime/CoreUObject/Private/UObject/PropertyArray.cpp
--- a/Engine/Source/Runtime/CoreUObject/Private/UObject/PropertyArray.cpp
+++ b/Engine/Source/Runtime/CoreUObject/Private/UObject/PropertyArray.cpp
@@ -88,7 +88,7 @@ void UArrayProperty::Serialize( FArchive& Ar )
 {
 	Super::Serialize( Ar );
 	Ar << Inner;
-	checkSlow(Inner||HasAnyFlags(RF_ClassDefaultObject));
+	checkSlow(Inner || HasAnyFlags(RF_ClassDefaultObject | RF_PendingKill));
 }
 void UArrayProperty::AddReferencedObjects(UObject* InThis, FReferenceCollector& Collector)
 {
