#UE4
Fix for TTP334025. We were corrupting Composite nodes at compile time when expanding instancing macros.

diff --git a/Engine/Source/Editor/KismetCompiler/Private/KismetCompiler.cpp b/Engine/Source/Editor/KismetCompiler/Private/KismetCompiler.cpp
--- a/Engine/Source/Editor/KismetCompiler/Private/KismetCompiler.cpp
+++ b/Engine/Source/Editor/KismetCompiler/Private/KismetCompiler.cpp
@@ -2686,16 +2686,30 @@ void FKismetCompilerContext::ExpandTunnelsAndMacros(UEdGraph* SourceGraph)
 					DuplicatedNode->NodePosY += NodeOffsetY;
 					DuplicatedNode->NodePosX += NodeOffsetX;
 
-					// Fix up the tunnel nodes to point correctly
+					if (const UK2Node_Composite* const CompositeNode = Cast<const UK2Node_Composite>(DuplicatedNode))
+					{
+						// Composite nodes can be present in the MacroNodes if users have collapsed nodes in the macro.
+						// No need to do anything for those:
+						continue;
+					}
+					
 					if (UK2Node_Tunnel* TunnelNode = Cast<UK2Node_Tunnel>(DuplicatedNode))
 					{
+						// Tunnel nodes should be connected to the MacroInstance they have been instantiated by.. Note 
+						// that if there are tunnel nodes internal to the macro instance they will be incorrectly 
+						// connected to the MacroInstance.
 						if (TunnelNode->bCanHaveInputs)
 						{
+							check(!TunnelNode->bCanHaveOutputs);
+							// If this check fails it indicates that we've failed to identify all uses of tunnel nodes and 
+							// are erroneously connecting tunnels to the macro instance when they should be left untouched..
+							check(!TunnelNode->InputSinkNode);
 							TunnelNode->InputSinkNode = MacroInstanceNode;
 							MacroInstanceNode->OutputSourceNode = TunnelNode;
 						}
 						else if (TunnelNode->bCanHaveOutputs)
 						{
+							check(!TunnelNode->OutputSourceNode);
 							TunnelNode->OutputSourceNode = MacroInstanceNode;
 							MacroInstanceNode->InputSinkNode = TunnelNode;
 						}
