UDS: member variables in UDS have no category metadata (previously the category was taken from PinType, it was an error).

diff --git a/Engine/Source/Editor/KismetCompiler/Private/UserDefinedStructureCompilerUtils.cpp b/Engine/Source/Editor/KismetCompiler/Private/UserDefinedStructureCompilerUtils.cpp
--- a/Engine/Source/Editor/KismetCompiler/Private/UserDefinedStructureCompilerUtils.cpp
+++ b/Engine/Source/Editor/KismetCompiler/Private/UserDefinedStructureCompilerUtils.cpp
@@ -156,7 +156,6 @@ struct FUserDefinedStructureCompilerInner
 				NewProperty->SetMetaData(FEdMode::MD_MakeEditWidget, TEXT("true"));
 			}
 			NewProperty->SetMetaData(TEXT("DisplayName"), *VarDesc.FriendlyName);
-			NewProperty->SetMetaData(TEXT("Category"), *VarDesc.Category);
 			NewProperty->SetMetaData(FBlueprintMetadata::MD_Tooltip, *VarDesc.ToolTip);
 			NewProperty->RepNotifyFunc = NAME_None;
 
