Minor: fix clang warning (ordering issues).

diff --git a/Engine/Source/Editor/KismetCompiler/Private/AnimBlueprintCompiler.cpp b/Engine/Source/Editor/KismetCompiler/Private/AnimBlueprintCompiler.cpp
--- a/Engine/Source/Editor/KismetCompiler/Private/AnimBlueprintCompiler.cpp
+++ b/Engine/Source/Editor/KismetCompiler/Private/AnimBlueprintCompiler.cpp
@@ -613,9 +613,9 @@ void FAnimBlueprintCompiler::ProcessStateMachine(UAnimGraphNode_StateMachineBase
 	public:
 		FMachineCreator(FCompilerResultsLog& InMessageLog, UAnimGraphNode_StateMachineBase* InStateMachineInstance, int32 InMachineIndex, UAnimBlueprintGeneratedClass* InNewClass)
 			: MachineIndex(InMachineIndex)
-			, MessageLog(InMessageLog)
-			, StateMachineInstance(InStateMachineInstance)
 			, AnimBlueprintClass(InNewClass)
+			, StateMachineInstance(InStateMachineInstance)
+			, MessageLog(InMessageLog)
 		{
 			FStateMachineDebugData& MachineInfo = GetMachineSpecificDebugData();
 			MachineInfo.MachineIndex = MachineIndex;
