BP should not crash when timeline is null.

diff --git a/Engine/Source/Editor/KismetCompiler/Private/KismetCompiler.cpp b/Engine/Source/Editor/KismetCompiler/Private/KismetCompiler.cpp
--- a/Engine/Source/Editor/KismetCompiler/Private/KismetCompiler.cpp
+++ b/Engine/Source/Editor/KismetCompiler/Private/KismetCompiler.cpp
@@ -201,7 +201,10 @@ void FKismetCompilerContext::CleanAndSanitizeClass(UBlueprintGeneratedClass* Cla
 			TSet<class UCurveBase*> Curves;
 			for (auto Timeline : Blueprint->Timelines)
 			{
-				Timeline->GetAllCurves(Curves);
+				if (Timeline)
+				{
+					Timeline->GetAllCurves(Curves);
+				}
 			}
 			for (auto Component : Blueprint->ComponentTemplates)
 			{
@@ -501,6 +504,12 @@ void FKismetCompilerContext::CreateClassVariablesFromBlueprint()
 	for (int32 TimelineIndex = 0; TimelineIndex < Blueprint->Timelines.Num(); ++TimelineIndex)
 	{
 		UTimelineTemplate* Timeline = Blueprint->Timelines[TimelineIndex];
+		// Not fatal if NULL, but shouldn't happen
+		if (!Timeline)
+		{
+			continue;
+		}
+
 		FEdGraphPinType TimelinePinType(Schema->PC_Object, TEXT(""), UTimelineComponent::StaticClass(), false, false);
 
 		// Previously UTimelineComponent object has exactly the same name as UTimelineTemplate object (that obj was in blueprint)
@@ -3036,6 +3045,15 @@ void FKismetCompilerContext::Compile()
 		}
 	}
 
+	for (int32 TimelineIndex = 0; TimelineIndex < Blueprint->Timelines.Num(); )
+	{
+		if (NULL == Blueprint->Timelines[TimelineIndex])
+		{
+			Blueprint->Timelines.RemoveAt(TimelineIndex);
+		}
+		++TimelineIndex;
+	}
+
 	CleanAndSanitizeClass(TargetClass, OldCDO);
 
 	FKismetCompilerVMBackend Backend_VM(Blueprint, Schema, *this);
diff --git a/Engine/Source/Editor/UnrealEd/Private/Kismet2/BlueprintEditorUtils.cpp b/Engine/Source/Editor/UnrealEd/Private/Kismet2/BlueprintEditorUtils.cpp
--- a/Engine/Source/Editor/UnrealEd/Private/Kismet2/BlueprintEditorUtils.cpp
+++ b/Engine/Source/Editor/UnrealEd/Private/Kismet2/BlueprintEditorUtils.cpp
@@ -2959,9 +2959,12 @@ void FBlueprintEditorUtils::GetClassVariableList(const UBlueprint* Blueprint, TA
 				VisibleVariables.AddUnique(ParentBP->NewVariables[VariableIndex].VarName);
 			}
 
-			for(int32 TimelineIndex = 0; TimelineIndex < ParentBP->Timelines.Num(); ++TimelineIndex)
+			for (auto Timeline : ParentBP->Timelines)
 			{
-				VisibleVariables.AddUnique(ParentBP->Timelines[TimelineIndex]->GetFName());
+				if (Timeline)
+				{
+					VisibleVariables.AddUnique(Timeline->GetFName());
+				}
 			}
 		}
 	}
