Skipping NULL timelines in BP. Fix.

diff --git a/Engine/Source/Editor/KismetCompiler/Private/KismetCompiler.cpp b/Engine/Source/Editor/KismetCompiler/Private/KismetCompiler.cpp
--- a/Engine/Source/Editor/KismetCompiler/Private/KismetCompiler.cpp
+++ b/Engine/Source/Editor/KismetCompiler/Private/KismetCompiler.cpp
@@ -3050,6 +3050,7 @@ void FKismetCompilerContext::Compile()
 		if (NULL == Blueprint->Timelines[TimelineIndex])
 		{
 			Blueprint->Timelines.RemoveAt(TimelineIndex);
+			continue;
 		}
 		++TimelineIndex;
 	}
