#ue4 Making the ticking error log more verbose in the case that an actor has a tick function, but cannot tick.

diff --git a/Engine/Source/Editor/KismetCompiler/Private/KismetCompiler.cpp b/Engine/Source/Editor/KismetCompiler/Private/KismetCompiler.cpp
--- a/Engine/Source/Editor/KismetCompiler/Private/KismetCompiler.cpp
+++ b/Engine/Source/Editor/KismetCompiler/Private/KismetCompiler.cpp
@@ -3443,7 +3443,7 @@ void FKismetCompilerContext::SetCanEverTickForActor()
 		else if(!CDActor->PrimaryActorTick.bCanEverTick)
 		{
 			const FString ReceivTickEventWarning = FString::Printf( 
-				*LOCTEXT("ReceiveTick_CanNeverTick", "Blueprint %s has the ReceiveTick @@ event, but it can never tick").ToString(), *NewClass->GetName());
+				*LOCTEXT("ReceiveTick_CanNeverTick", "Blueprint %s has the ReceiveTick @@ event, but it can never tick.  Please consider using a Timer instead of a tick, using Actor as the parent class, or you can enable Tick using code in one of the following ways: set ChildCanTick in the metadata on the parent class, or set bCanEverTick to true.").ToString(), *NewClass->GetName());
 			MessageLog.Warning( *ReceivTickEventWarning, FindLocalEntryPoint(ReciveTickEvent) );
 		}
 	}
