GC's DebugTokenMap wasn't emptied during Blueprints recompile while ReferenceTokenStream was.

diff --git a/Engine/Source/Editor/KismetCompiler/Private/KismetCompiler.cpp b/Engine/Source/Editor/KismetCompiler/Private/KismetCompiler.cpp
--- a/Engine/Source/Editor/KismetCompiler/Private/KismetCompiler.cpp
+++ b/Engine/Source/Editor/KismetCompiler/Private/KismetCompiler.cpp
@@ -1588,6 +1588,9 @@ void FKismetCompilerContext::FinishCompilingClass(UClass* Class)
 	{
 		// Propagate the new parent's inheritable class flags
 		Class->ReferenceTokenStream.Empty();
+#if !(UE_BUILD_TEST || UE_BUILD_SHIPPING)
+		Class->DebugTokenMap.Empty();
+#endif
 		Class->ClassFlags &= ~CLASS_RecompilerClear;
 		Class->ClassFlags |= (ParentClass->ClassFlags & CLASS_ScriptInherit);//@TODO: ChangeParentClass had this, but I don't think I want it: | UClass::StaticClassFlags;  // will end up with CLASS_Intrinsic
 		Class->ClassCastFlags |= ParentClass->ClassCastFlags;
diff --git a/Engine/Source/Runtime/CoreUObject/Public/UObject/GarbageCollection.h b/Engine/Source/Runtime/CoreUObject/Public/UObject/GarbageCollection.h
--- a/Engine/Source/Runtime/CoreUObject/Public/UObject/GarbageCollection.h
+++ b/Engine/Source/Runtime/CoreUObject/Public/UObject/GarbageCollection.h
@@ -175,6 +175,14 @@ public:
 	 */
 	const FTokenInfo& GetTokenInfo(int32 TokenIndex) const;
 
+	/**
+	 * Empties the map.
+	 */
+	void Empty()
+	{
+		TokenMap.Empty();
+	}
+
 private:
 	/* Token map. */
 	TArray<FTokenInfo> TokenMap;
