Allow blueprint struct pins to be split in to their component elements on the node instead of requiring a make/break node be used.

diff --git a/Engine/Source/Editor/BlueprintGraph/Classes/EdGraphSchema_K2.h b/Engine/Source/Editor/BlueprintGraph/Classes/EdGraphSchema_K2.h
--- a/Engine/Source/Editor/BlueprintGraph/Classes/EdGraphSchema_K2.h
+++ b/Engine/Source/Editor/BlueprintGraph/Classes/EdGraphSchema_K2.h
@@ -17,9 +17,9 @@ struct BLUEPRINTGRAPH_API FEdGraphSchemaAction_K2Struct : public FEdGraphSchemaA
 	UStruct* Struct;
 
 	void AddReferencedObjects( FReferenceCollector& Collector ) override
-	{
+{
 		if( Struct )
-		{
+	{
 			Collector.AddReferencedObject(Struct);
 		}
 	}
@@ -45,7 +45,7 @@ public:
 	// Struct/Enum/Class:
 	// If true, this class, struct, or enum is a valid type for use as a variable in a blueprint
 	static const FName MD_AllowableBlueprintVariableType;
-	
+
 	// If true, this class, struct, or enum is not valid for use as a variable in a blueprint
 	static const FName MD_NotAllowableBlueprintVariableType;
 
@@ -77,7 +77,7 @@ public:
 
 	// [FunctionMetadata] Indicates that the function should be drawn as a compact node with the specified body title
 	static const FName MD_CompactNodeTitle;
-	
+
 	// [FunctionMetadata] Indicates that the function should be drawn with this title over the function name
 	static const FName MD_FriendlyName;
 
@@ -233,13 +233,13 @@ class BLUEPRINTGRAPH_API UEdGraphSchema_K2 : public UEdGraphSchema
 	FString PC_Name;
 
 	UPROPERTY()
-	FString PC_Delegate;	// SubCategoryObject is the UFunction of the delegate signature
+	FString PC_Delegate;    // SubCategoryObject is the UFunction of the delegate signature
 
 	UPROPERTY()
-	FString PC_MCDelegate;	// SubCategoryObject is the UFunction of the delegate signature
+	FString PC_MCDelegate;  // SubCategoryObject is the UFunction of the delegate signature
 
 	UPROPERTY()
-	FString PC_Object;		// SubCategoryObject is the Class of the object passed thru this pin, or SubCategory can be 'self'. The DefaultValue string should always be empty, use DefaultObject.
+	FString PC_Object;    // SubCategoryObject is the Class of the object passed thru this pin, or SubCategory can be 'self'. The DefaultValue string should always be empty, use DefaultObject.
 
 	UPROPERTY()
 	FString PC_Interface;	// SubCategoryObject is the Class of the object passed thru this pin.
@@ -254,7 +254,7 @@ class BLUEPRINTGRAPH_API UEdGraphSchema_K2 : public UEdGraphSchema
 	FString PC_Struct;    // SubCategoryObject is the ScriptStruct of the struct passed thru this pin, 'self' is not a valid SubCategory. DefaultObject should always be empty, the DefaultValue string may be used for supported structs.
 
 	UPROPERTY()
-	FString PC_Wildcard;  // Special matching rules are imposed by the node itself
+	FString PC_Wildcard;    // Special matching rules are imposed by the node itself
 
 	// Common PinType.PinSubCategory values
 	UPROPERTY()
@@ -351,7 +351,7 @@ class BLUEPRINTGRAPH_API UEdGraphSchema_K2 : public UEdGraphSchema
 	/** Whether or not the schema should allow the user to use blueprint communications */
 	UPROPERTY(globalconfig)
 	bool bAllowBlueprintComms;
-	
+
 	UPROPERTY(globalconfig)
 	TArray<FBlueprintCallableFunctionRedirect> EditoronlyBPFunctionRedirects;
 
@@ -493,8 +493,17 @@ public:
 	virtual TSharedPtr<FEdGraphSchemaAction> GetCreateCommentAction() const override;
 	virtual bool FadeNodeWhenDraggingOffPin(const UEdGraphNode* Node, const UEdGraphPin* Pin) const override;
 	virtual void BackwardCompatibilityNodeConversion(UEdGraph* Graph, bool bOnlySafeChanges) const override;
+	virtual void SplitPin(UEdGraphPin* Pin) const override;
+	virtual void RecombinePin(UEdGraphPin* Pin) const override;
 	// End EdGraphSchema Interface
 
+	// Returns whether the supplied Pin is a splittable struct
+	bool PinHasSplittableStructType(const UEdGraphPin* InGraphPin) const;
+
+	/** Helper function to create the expansion node.  
+		If the CompilerContext is specified this will be created as an intermediate node */
+	class UK2Node* CreateSplitPinNode(UEdGraphPin* Pin, class FKismetCompilerContext* CompilerContext = NULL, UEdGraph* SourceGraph = NULL) const;
+
 	// Do validation, that doesn't require a knowledge about actual pin. 
 	virtual bool DefaultValueSimpleValidation(const FEdGraphPinType& PinType, const FString& PinName, const FString& NewDefaultValue, UObject* NewDefaultObject, const FText& InText, FString* OutMsg = NULL) const;
 
@@ -614,6 +623,12 @@ public:
 	/** Can Pin be promoted to a variable? */
 	bool CanPromotePinToVariable (const UEdGraphPin& Pin) const;
 
+	/** Can Pin be split in to its component elements */
+	bool CanSplitStructPin(const UEdGraphPin& Pin) const;
+
+	/** Can Pin be recombined back to its original form */
+	bool CanRecombineStructPin(const UEdGraphPin& Pin) const;
+
 	/**
 	 * Convert the type of a UProperty to the corresponding pin type.
 	 *
@@ -632,7 +647,7 @@ public:
 		FT_Const		= 0x04,
 		FT_Protected	= 0x08,
 	};
-	
+
 	/**
 	 * Finds the parent function for the specified function, if any
 	 *
@@ -752,7 +767,7 @@ public:
 	 */
 	virtual void CreateFunctionGraphTerminators(UEdGraph& Graph, UClass* Class) const;
 
-	/** 
+	/**
 	 * Populate new function graph with entry and possibly return node
 	 * 
 	 * @param	Graph			Graph to add the function terminators to
@@ -859,6 +874,9 @@ public:
 	 */
 	virtual bool ArePinTypesCompatible(const FEdGraphPinType& Output, const FEdGraphPinType& Input, const UClass* CallingContext = NULL, bool bIgnoreArray = false) const;
 
+	// Sets the Pin default property potentially using the default value set in metadata of the param
+	virtual void SetPinDefaultValue(UEdGraphPin* Pin, const UFunction* Function = NULL, const UProperty* Param = NULL) const;
+
 	/**
 	 * Sets the default value of a pin based on the type of the pin (0 for int, false for boolean, etc...)
 	 */
@@ -894,7 +912,7 @@ public:
 	/** Calculates an estimated height for the specified node */
 	static float EstimateNodeHeight( UEdGraphNode* Node );
 
-	/**
+	/** 
 	 * Checks if the graph supports impure functions
 	 *
 	 * @param InGraph		Graph to check
@@ -927,9 +945,9 @@ public:
 	 *
 	 * @return						Returns TRUE if successful
 	 */
-	bool CollapseGatewayNode(UK2Node* InNode, UEdGraphNode* InEntryNode, UEdGraphNode* InResultNode) const;
+	bool CollapseGatewayNode(UK2Node* InNode, UEdGraphNode* InEntryNode, UEdGraphNode* InResultNode, class FKismetCompilerContext* CompilerContext = NULL) const;
 
-	/**
+	/** 
 	 * Connects all of the linked pins from PinA to all of the linked pins from PinB, removing
 	 * both PinA and PinB from being linked to anything else
 	 * Requires the nodes that own the pins to be in the same graph already (post-merging)
diff --git a/Engine/Source/Editor/BlueprintGraph/Classes/K2Node.h b/Engine/Source/Editor/BlueprintGraph/Classes/K2Node.h
--- a/Engine/Source/Editor/BlueprintGraph/Classes/K2Node.h
+++ b/Engine/Source/Editor/BlueprintGraph/Classes/K2Node.h
@@ -201,9 +201,11 @@ class UK2Node : public UEdGraphNode
 	 */
 	bool CreatePinsForFunctionEntryExit(const UFunction* Function, bool bForFunctionEntry);
 
-	BLUEPRINTGRAPH_API virtual void ExpandNode(class FKismetCompilerContext& CompilerContext, UEdGraph* SourceGraph) { };
+	BLUEPRINTGRAPH_API virtual void ExpandNode(class FKismetCompilerContext& CompilerContext, UEdGraph* SourceGraph);
 	BLUEPRINTGRAPH_API virtual class FNodeHandlingFunctor* CreateNodeHandler(class FKismetCompilerContext& CompilerContext) const { return NULL; }
 
+	BLUEPRINTGRAPH_API void ExpandSplitPin(class FKismetCompilerContext* CompilerContext, UEdGraph* SourceGraph, UEdGraphPin* Pin);
+
 	/**
 	 * Query if this is a node that is safe to ignore (e.g., a comment node or other non-structural annotation that can be pruned with no warnings).
 	 *
diff --git a/Engine/Source/Editor/BlueprintGraph/Private/K2ActionMenuBuilder.cpp b/Engine/Source/Editor/BlueprintGraph/Private/K2ActionMenuBuilder.cpp
--- a/Engine/Source/Editor/BlueprintGraph/Private/K2ActionMenuBuilder.cpp
+++ b/Engine/Source/Editor/BlueprintGraph/Private/K2ActionMenuBuilder.cpp
@@ -2533,11 +2533,11 @@ void FK2ActionMenuBuilder::GetFuncNodesWithPinType(FBlueprintGraphActionListBuil
 
 		static const FName NativeMakeFunc(TEXT("NativeMakeFunc"));
 		static const FName NativeBrakeFunc(TEXT("NativeBreakFunc"));
-		static const FName HasNativeMakeBreak(TEXT("HasNativeMakeBreak"));
+		static const FName HasNativeMake(TEXT("HasNativeMake"));
+		static const FName HasNativeBreak(TEXT("HasNativeBreak"));
 		const UScriptStruct* ScriptStruct = Cast<const UScriptStruct>(DesiredPinType.PinSubCategoryObject.Get());
-		const bool bStructWithNativeBrakeAndMake = ScriptStruct && ScriptStruct->GetBoolMetaData(HasNativeMakeBreak);
-		const bool bUseNativeMake = bStructWithNativeBrakeAndMake && bWantOutput;
-		const bool bUseNativeBrake = bStructWithNativeBrakeAndMake && !bWantOutput;
+		const bool bUseNativeMake = bWantOutput && ScriptStruct && ScriptStruct->GetBoolMetaData(HasNativeMake);
+		const bool bUseNativeBrake = !bWantOutput && ScriptStruct && ScriptStruct->GetBoolMetaData(HasNativeBreak);
 
 		for (TFieldIterator<UFunction> FunctionIt(Class, EFieldIteratorFlags::IncludeSuper); FunctionIt; ++FunctionIt)
 		{
diff --git a/Engine/Source/Editor/BlueprintGraph/Private/K2Node.cpp b/Engine/Source/Editor/BlueprintGraph/Private/K2Node.cpp
--- a/Engine/Source/Editor/BlueprintGraph/Private/K2Node.cpp
+++ b/Engine/Source/Editor/BlueprintGraph/Private/K2Node.cpp
@@ -432,17 +432,78 @@ UK2Node::ERedirectType UK2Node::DoPinsMatchForReconstruction(const UEdGraphPin*
 		// try looking for a redirect if it's a K2 node
 		if (UK2Node* Node = Cast<UK2Node>(NewPin->GetOwningNode()))
 		{	
-			// if you don't have matching pin, now check if there is any redirect param set
-			TArray<FString> OldPinNames;
-			GetRedirectPinNames(*OldPin, OldPinNames);
+			if (OldPin->ParentPin == NULL)
+			{
+				// if you don't have matching pin, now check if there is any redirect param set
+				TArray<FString> OldPinNames;
+				GetRedirectPinNames(*OldPin, OldPinNames);
 
-			FName NewPinName;
-			RedirectType = ShouldRedirectParam(OldPinNames, /*out*/ NewPinName, Node);
+				FName NewPinName;
+				RedirectType = ShouldRedirectParam(OldPinNames, /*out*/ NewPinName, Node);
 
-			// make sure they match
-			if ((RedirectType != ERedirectType_None) && FCString::Stricmp(*(NewPin->PinName), *(NewPinName.ToString())) != 0)
+				// make sure they match
+				if ((RedirectType != ERedirectType_None) && FCString::Stricmp(*(NewPin->PinName), *(NewPinName.ToString())) != 0)
+				{
+					RedirectType = ERedirectType_None;
+				}
+			}
+			else
 			{
-				RedirectType = ERedirectType_None;
+				struct FPropertyDetails
+				{
+					const UEdGraphPin* Pin;
+					FString PropertyName;
+
+					FPropertyDetails(const UEdGraphPin* InPin, const FString& InPropertyName)
+						: Pin(InPin), PropertyName(InPropertyName)
+					{
+					}
+				};
+
+				TArray<FPropertyDetails> ParentHierarchy;
+				const UEdGraphPin* CurPin = OldPin;
+				do 
+				{
+					ParentHierarchy.Add(FPropertyDetails(CurPin, CurPin->PinName.RightChop(CurPin->ParentPin->PinName.Len() + 1)));
+					CurPin = CurPin->ParentPin;
+				} while (CurPin->ParentPin);
+
+				// if you don't have matching pin, now check if there is any redirect param set
+				TArray<FString> OldPinNames;
+				GetRedirectPinNames(*CurPin, OldPinNames);
+
+				FString NewPinNameStr;
+				FName NewPinName;
+				RedirectType = ShouldRedirectParam(OldPinNames, /*out*/ NewPinName, Node);
+
+				NewPinNameStr = (RedirectType == ERedirectType_None ? CurPin->PinName : NewPinName.ToString());
+
+				for (int32 ParentIndex = ParentHierarchy.Num() - 1; ParentIndex >= 0; --ParentIndex)
+				{
+					const UEdGraphPin* CurPin = ParentHierarchy[ParentIndex].Pin;
+					const UEdGraphPin* ParentPin = CurPin->ParentPin;
+
+					TMap<FName, FName>* StructRedirects = UStruct::TaggedPropertyRedirects.Find(ParentPin->PinType.PinSubCategoryObject->GetFName());
+					if (StructRedirects)
+					{
+						FName* PropertyRedirect = StructRedirects->Find(FName(*ParentHierarchy[ParentIndex].PropertyName));
+						if (PropertyRedirect)
+						{
+							NewPinNameStr += FString("_") + PropertyRedirect->ToString();
+						}
+						else
+						{
+							NewPinNameStr += FString("_") + ParentHierarchy[ParentIndex].PropertyName;
+						}
+					}
+					else
+					{
+						NewPinNameStr += FString("_") + ParentHierarchy[ParentIndex].PropertyName;
+					}
+				}
+
+				// make sure they match
+				RedirectType = ((FCString::Stricmp(*(NewPin->PinName), *NewPinNameStr) != 0) ? ERedirectType_None : ERedirectType_Name);
 			}
 		}
 	}
@@ -513,6 +574,60 @@ void UK2Node::ReconstructSinglePin(UEdGraphPin* NewPin, UEdGraphPin* OldPin, ERe
 	OldPin->Rename(NULL, GetTransientPackage(), (REN_DontCreateRedirectors|(Blueprint->bIsRegeneratingOnLoad ? REN_ForceNoResetLoaders : REN_None)));
 }
 
+void UK2Node::ExpandSplitPin(FKismetCompilerContext* CompilerContext, UEdGraph* SourceGraph, UEdGraphPin* Pin)
+{
+	const UEdGraphSchema_K2* Schema = CastChecked<UEdGraphSchema_K2>(CompilerContext ? CompilerContext->GetSchema() : SourceGraph->GetSchema());
+
+	UK2Node* ExpandedNode = Schema->CreateSplitPinNode(Pin, CompilerContext, SourceGraph);
+
+	int32 SubPinIndex = 0;
+
+	for (int32 ExpandedPinIndex=0; ExpandedPinIndex < ExpandedNode->Pins.Num(); ++ExpandedPinIndex)
+	{
+		UEdGraphPin* ExpandedPin = ExpandedNode->Pins[ExpandedPinIndex];
+
+		if (!ExpandedPin->bHidden)
+		{
+			if (ExpandedPin->Direction == Pin->Direction)
+			{
+				UEdGraphPin* SubPin = Pin->SubPins[SubPinIndex++];
+				if (CompilerContext)
+				{
+					CompilerContext->MovePinLinksToIntermediate(*SubPin, *ExpandedPin);
+				}
+				else
+				{
+					Schema->MovePinLinks(*SubPin, *ExpandedPin);
+				}
+				Pins.Remove(SubPin);
+			}
+			else
+			{
+				Schema->TryCreateConnection(Pin, ExpandedPin);
+			}
+		}
+	}
+
+	Pin->SubPins.Empty();				
+}
+
+void UK2Node::ExpandNode(FKismetCompilerContext& CompilerContext, UEdGraph* SourceGraph)
+{
+	if (CompilerContext.bIsFullCompile)
+	{
+		// We iterate the array in reverse so we can both remove the subpins safely after we've read them and
+		// so we have split nested structs we combine them back together in the right order
+		for (int32 PinIndex=Pins.Num() - 1; PinIndex >= 0; --PinIndex)
+		{
+			UEdGraphPin* Pin = Pins[PinIndex];
+			if (Pin->SubPins.Num() > 0)
+			{
+				ExpandSplitPin(&CompilerContext, SourceGraph, Pin);
+			}
+		}
+	}
+}
+
 bool UK2Node::HasValidBlueprint() const
 {
 	// Perform an unchecked search here, so we don't crash if this is a transient node from a list refresh without a valid outer blueprint
diff --git a/Engine/Source/Editor/BlueprintGraph/Private/K2Node_BaseAsyncTask.cpp b/Engine/Source/Editor/BlueprintGraph/Private/K2Node_BaseAsyncTask.cpp
--- a/Engine/Source/Editor/BlueprintGraph/Private/K2Node_BaseAsyncTask.cpp
+++ b/Engine/Source/Editor/BlueprintGraph/Private/K2Node_BaseAsyncTask.cpp
@@ -42,7 +42,7 @@ void UK2Node_BaseAsyncTask::GetMenuEntries(FGraphContextMenuBuilder& ContextMenu
 	const bool bAllowLatentFuncs = (K2Schema->GetGraphType(ContextMenuBuilder.CurrentGraph) == GT_Ubergraph);
 
 	if (bAllowLatentFuncs)
-	{
+{
 		UK2Node_BaseAsyncTask* NodeTemplate = NewObject<UK2Node_BaseAsyncTask>(ContextMenuBuilder.OwnerOfTemporaries, GetClass());
 		CreateDefaultMenuEntry(NodeTemplate, ContextMenuBuilder);
 	}
@@ -70,13 +70,13 @@ void UK2Node_BaseAsyncTask::AllocateDefaultPins()
 	{
 		if (UMulticastDelegateProperty* Property = Cast<UMulticastDelegateProperty>(*PropertyIt))
 		{
-			CreatePin(EGPD_Output, K2Schema->PC_Exec, TEXT(""), NULL, false, false, *Property->GetName());
-			if (!DelegateSignatureFunction)
-			{
-				DelegateSignatureFunction = Property->SignatureFunction;
-			}
+		CreatePin(EGPD_Output, K2Schema->PC_Exec, TEXT(""), NULL, false, false, *Property->GetName());
+		if (!DelegateSignatureFunction)
+		{
+			DelegateSignatureFunction = Property->SignatureFunction;
 		}
 	}
+	}
 
 	if (DelegateSignatureFunction)
 	{
@@ -124,26 +124,7 @@ void UK2Node_BaseAsyncTask::AllocateDefaultPins()
 					AdvancedPinDisplay = ENodeAdvancedPins::Hidden;
 				}
 
-				FString MetadataDefaultValue = Function->GetMetaData(*Param->GetName());
-				if (!MetadataDefaultValue.IsEmpty())
-				{
-					// Specified default value in the metadata
-					Pin->DefaultValue = Pin->AutogeneratedDefaultValue = MetadataDefaultValue;
-				}
-				else
-				{
-					const FName MetadataCppDefaultValueKey( *(FString(TEXT("CPP_Default_")) + Param->GetName()) );
-					const FString MetadataCppDefaultValue = Function->GetMetaData(MetadataCppDefaultValueKey);
-					if(!MetadataCppDefaultValue.IsEmpty())
-					{
-						Pin->DefaultValue = Pin->AutogeneratedDefaultValue = MetadataCppDefaultValue;
-					}
-					else
-					{
-						// Set the default value to (T)0
-						K2Schema->SetPinDefaultValueBasedOnType(Pin);
-					}
-				}
+				K2Schema->SetPinDefaultValue(Pin, Function, Param);
 
 				if (PinsToHide.Contains(Pin->PinName))
 				{
@@ -163,7 +144,7 @@ void UK2Node_BaseAsyncTask::AllocateDefaultPins()
 struct FBaseAsyncTaskHelper
 {
 	static bool ValidDataPin(const UEdGraphPin* Pin, EEdGraphPinDirection Direction, const UEdGraphSchema_K2* Schema)
-	{
+{
 		check(Schema);
 		const bool bValidDataPin = Pin
 			&& (Pin->PinName != Schema->PN_Execute)
@@ -173,11 +154,11 @@ struct FBaseAsyncTaskHelper
 		const bool bProperDirection = Pin && (Pin->Direction == Direction);
 
 		return bValidDataPin && bProperDirection;
-	}
+}
 
 	static bool CreateDelegateForNewFunction(UEdGraphPin* DelegateInputPin, FName FunctionName, UK2Node* CurrentNode, UEdGraph* SourceGraph, FKismetCompilerContext& CompilerContext)
 	{
-		const UEdGraphSchema_K2* Schema = CompilerContext.GetSchema();
+	const UEdGraphSchema_K2* Schema = CompilerContext.GetSchema();
 		check(DelegateInputPin && Schema && CurrentNode && SourceGraph && (FunctionName != NAME_None));
 		bool bResult = true;
 
@@ -212,12 +193,12 @@ struct FBaseAsyncTaskHelper
 		{
 			const UProperty* Param = *PropIt;
 			if (!Param->HasAnyPropertyFlags(CPF_OutParm) || Param->HasAnyPropertyFlags(CPF_ReferenceParm))
-			{
+		{
 				FEdGraphPinType PinType;
 				bResult &= Schema->ConvertPropertyToPinType(Param, /*out*/ PinType);
 				bResult &= (NULL != CENode->CreateUserDefinedPin(Param->GetName(), PinType));
-			}
 		}
+	}
 
 		return bResult;
 	}
@@ -241,9 +222,9 @@ struct FBaseAsyncTaskHelper
 
 		UK2Node_CustomEvent* CurrentCENode = CompilerContext.SpawnIntermediateNode<UK2Node_CustomEvent>(CurrentNode, SourceGraph);
 		{
-			UK2Node_AddDelegate* AddDelegateNode = CompilerContext.SpawnIntermediateNode<UK2Node_AddDelegate>(CurrentNode, SourceGraph);
-			AddDelegateNode->SetFromProperty(CurrentProperty, false);
-			AddDelegateNode->AllocateDefaultPins();
+		UK2Node_AddDelegate* AddDelegateNode = CompilerContext.SpawnIntermediateNode<UK2Node_AddDelegate>(CurrentNode, SourceGraph);
+		AddDelegateNode->SetFromProperty(CurrentProperty, false);
+		AddDelegateNode->AllocateDefaultPins();
 			bIsErrorFree &= Schema->TryCreateConnection(AddDelegateNode->FindPinChecked(Schema->PN_Self), ProxyObjectPin);
 			bIsErrorFree &= Schema->TryCreateConnection(InOutLastThenPin, AddDelegateNode->FindPinChecked(Schema->PN_Execute));
 			InOutLastThenPin = AddDelegateNode->FindPinChecked(Schema->PN_Then);
@@ -259,11 +240,11 @@ struct FBaseAsyncTaskHelper
 		{
 			UEdGraphPin* PinWithData = CurrentCENode->FindPin(OutputPair.OutputPin->PinName);
 			if (PinWithData == NULL)
-			{
+		{
 				FText ErrorMessage = FText::Format(LOCTEXT("MissingDataPin", "ICE: Pin @@ was expecting a data output pin named {0} on @@ (each delegate must have the same signature)"), FText::FromString(OutputPair.OutputPin->PinName));
 				CompilerContext.MessageLog.Error(*ErrorMessage.ToString(), OutputPair.OutputPin, CurrentCENode);
 				return false;
-			}
+		}
 
 			UK2Node_AssignmentStatement* AssignNode = CompilerContext.SpawnIntermediateNode<UK2Node_AssignmentStatement>(CurrentNode, SourceGraph);
 			AssignNode->AllocateDefaultPins();
@@ -283,6 +264,8 @@ struct FBaseAsyncTaskHelper
 
 void UK2Node_BaseAsyncTask::ExpandNode(class FKismetCompilerContext& CompilerContext, UEdGraph* SourceGraph)
 {
+        Super::ExpandNode(CompilerContext, SourceGraph);
+
 	if (!CompilerContext.bIsFullCompile)
 	{
 		return;
@@ -297,14 +280,14 @@ void UK2Node_BaseAsyncTask::ExpandNode(class FKismetCompilerContext& CompilerCon
 	CallCreateProxyObjectNode->AllocateDefaultPins();
 	bIsErrorFree &= CompilerContext.MovePinLinksToIntermediate(*FindPinChecked(Schema->PN_Execute), *CallCreateProxyObjectNode->FindPinChecked(Schema->PN_Execute)).CanSafeConnect();
 	for (auto CurrentPin : Pins)
-	{
-		if (FBaseAsyncTaskHelper::ValidDataPin(CurrentPin, EGPD_Input, Schema))
 		{
+		if (FBaseAsyncTaskHelper::ValidDataPin(CurrentPin, EGPD_Input, Schema))
+			{
 			UEdGraphPin* DestPin = CallCreateProxyObjectNode->FindPin(CurrentPin->PinName); // match function inputs, to pass data to function from CallFunction node
 			bIsErrorFree &= DestPin && CompilerContext.MovePinLinksToIntermediate(*CurrentPin, *DestPin).CanSafeConnect();
+			}
 		}
-	}
-	
+
 	// GATHER OUTPUT PARAMETERS AND PAIR THEM WITH LOCAL VARIABLES
 	TArray<FBaseAsyncTaskHelper::FOutputPinAndLocalVariable> VariableOutputs;
 	for (auto CurrentPin : Pins)
@@ -323,15 +306,15 @@ void UK2Node_BaseAsyncTask::ExpandNode(class FKismetCompilerContext& CompilerCon
 	UEdGraphPin* LastThenPin = CallCreateProxyObjectNode->FindPinChecked(Schema->PN_Then);
 	UEdGraphPin* const ProxyObjectPin = CallCreateProxyObjectNode->GetReturnValuePin();
 	for (TFieldIterator<UMulticastDelegateProperty> PropertyIt(ProxyClass, EFieldIteratorFlags::ExcludeSuper); PropertyIt && bIsErrorFree; ++PropertyIt)
-	{
+		{
 		bIsErrorFree &= FBaseAsyncTaskHelper::HandleDelegateImplementation(*PropertyIt, VariableOutputs, ProxyObjectPin, LastThenPin, this, SourceGraph, CompilerContext);
 	}
 
 	if (CallCreateProxyObjectNode->FindPinChecked(Schema->PN_Then) == LastThenPin)
-	{
+			{
 		CompilerContext.MessageLog.Error(*LOCTEXT("MissingDelegateProperties", "BaseAsyncTask: Proxy has no delegates defined. @@").ToString(), this);
 		return;
-	}
+			}
 
 	// Create a call to activate the proxy object if necessary
 	if (ProxyActivateFunctionName != NAME_None)
@@ -345,7 +328,7 @@ void UK2Node_BaseAsyncTask::ExpandNode(class FKismetCompilerContext& CompilerCon
 		check(ActivateCallSelfPin);
 
 		bIsErrorFree &= Schema->TryCreateConnection(ProxyObjectPin, ActivateCallSelfPin);
-		
+
 		// Hook the activate node up in the exec chain
 		UEdGraphPin* ActivateExecPin = CallActivateProxyObjectNode->FindPinChecked(Schema->PN_Execute);
 		UEdGraphPin* ActivateThenPin = CallActivateProxyObjectNode->FindPinChecked(Schema->PN_Then);
@@ -360,11 +343,11 @@ void UK2Node_BaseAsyncTask::ExpandNode(class FKismetCompilerContext& CompilerCon
 	if (!bIsErrorFree)
 	{
 		CompilerContext.MessageLog.Error(*LOCTEXT("InternalConnectionError", "BaseAsyncTask: Internal connection error. @@").ToString(), this);
-	}
+		}
 
 	// Make sure we caught everything
 	BreakAllNodeLinks();
-}
+	}
 
 bool UK2Node_BaseAsyncTask::HasExternalBlueprintDependencies(TArray<class UStruct*>* OptionalOutput) const
 {
@@ -386,9 +369,9 @@ bool UK2Node_BaseAsyncTask::HasExternalBlueprintDependencies(TArray<class UStruc
 }
 
 FName UK2Node_BaseAsyncTask::GetCornerIcon() const
-{
+	{
 	return TEXT("Graph.Latent.LatentIcon");
-}
+	}
 
 UFunction* UK2Node_BaseAsyncTask::GetFactoryFunction() const
 {
diff --git a/Engine/Source/Editor/BlueprintGraph/Private/K2Node_BreakStruct.cpp b/Engine/Source/Editor/BlueprintGraph/Private/K2Node_BreakStruct.cpp
--- a/Engine/Source/Editor/BlueprintGraph/Private/K2Node_BreakStruct.cpp
+++ b/Engine/Source/Editor/BlueprintGraph/Private/K2Node_BreakStruct.cpp
@@ -126,7 +126,7 @@ UK2Node_BreakStruct::UK2Node_BreakStruct(const class FPostConstructInitializePro
 bool UK2Node_BreakStruct::CanBeBroken(const UScriptStruct* Struct)
 {
 	const UEdGraphSchema_K2* Schema = GetDefault<UEdGraphSchema_K2>();
-	if(Struct && Schema && !Struct->GetBoolMetaData(TEXT("HasNativeMakeBreak")))
+	if(Struct && Schema && !Struct->HasMetaData(TEXT("HasNativeBreak")))
 	{
 		for (TFieldIterator<UProperty> It(Struct); It; ++It)
 		{
@@ -217,9 +217,28 @@ FLinearColor UK2Node_BreakStruct::GetNodeTitleColor() const
 UK2Node::ERedirectType UK2Node_BreakStruct::DoPinsMatchForReconstruction(const UEdGraphPin* NewPin, int32 NewPinIndex, const UEdGraphPin* OldPin, int32 OldPinIndex)  const
 {
 	ERedirectType Result = UK2Node::DoPinsMatchForReconstruction(NewPin, NewPinIndex, OldPin, OldPinIndex);
-	if ((ERedirectType_None == Result) && DoRenamedPinsMatch(NewPin, OldPin, true))
+	if ((ERedirectType_None == Result) && NewPin && OldPin)
 	{
-		Result = ERedirectType_Custom;
+		if ((EGPD_Input == NewPin->Direction) && (EGPD_Input == OldPin->Direction))
+	{
+		const UEdGraphSchema_K2* K2Schema = GetDefault<UEdGraphSchema_K2>();
+		if (K2Schema->ArePinTypesCompatible( NewPin->PinType, OldPin->PinType))
+		{
+			Result = ERedirectType_Custom;
+		}
+	}
+		else if ((EGPD_Output == NewPin->Direction) && (EGPD_Output == OldPin->Direction))
+		{
+			TMap<FName, FName>* StructRedirects = UStruct::TaggedPropertyRedirects.Find(StructType->GetFName());
+			if (StructRedirects)
+			{
+				FName* PropertyRedirect = StructRedirects->Find(FName(*OldPin->PinName));
+				if (PropertyRedirect)
+				{
+					Result = ((FCString::Stricmp(*PropertyRedirect->ToString(), *NewPin->PinName) != 0) ? ERedirectType_None : ERedirectType_Name);
+				}
+			}
+		}
 	}
 	return Result;
 }
diff --git a/Engine/Source/Editor/BlueprintGraph/Private/K2Node_CallFunction.cpp b/Engine/Source/Editor/BlueprintGraph/Private/K2Node_CallFunction.cpp
--- a/Engine/Source/Editor/BlueprintGraph/Private/K2Node_CallFunction.cpp
+++ b/Engine/Source/Editor/BlueprintGraph/Private/K2Node_CallFunction.cpp
@@ -152,13 +152,13 @@ FText UK2Node_CallFunction::GetNodeTitle(ENodeTitleType::Type TitleType) const
 		ContextString = GetFunctionContextString();
 	}
 	else
-	{
+		{
 		FunctionName = FunctionReference.GetMemberName().ToString();
 		if ((GEditor != NULL) && (GetDefault<UEditorStyleSettings>()->bShowFriendlyNames))
 		{
 			FunctionName = FName::NameToDisplayString(FunctionName, false);
 		}
-	}
+		}
 
 	if(TitleType == ENodeTitleType::FullTitle)
 	{
@@ -378,9 +378,9 @@ void UK2Node_CallFunction::CreateExecPinsForFunctionCall(const UFunction* Functi
 		if (bCreateThenPin)
 		{
 			UEdGraphPin* OutputExecPin = CreatePin(EGPD_Output, K2Schema->PC_Exec, TEXT(""), NULL, false, false, K2Schema->PN_Then);
-			// Use 'completed' name for output pins on latent functions
-			if(Function->HasMetaData(FBlueprintMetadata::MD_Latent))
-			{
+		// Use 'completed' name for output pins on latent functions
+		if(Function->HasMetaData(FBlueprintMetadata::MD_Latent))
+		{
 				OutputExecPin->PinFriendlyName = FText::FromString(K2Schema->PN_Completed);
 			}
 		}
@@ -609,26 +609,7 @@ bool UK2Node_CallFunction::CreatePinsForFunctionCall(const UFunction* Function)
 				AdvancedPinDisplay = ENodeAdvancedPins::Hidden;
 			}
 
-			FString MetadataDefaultValue = Function->GetMetaData(*Param->GetName());
-			if (!MetadataDefaultValue.IsEmpty())
-			{
-				// Specified default value in the metadata
-				Pin->DefaultValue = Pin->AutogeneratedDefaultValue = MetadataDefaultValue;
-			}
-			else
-			{
-				const FName MetadataCppDefaultValueKey( *(FString(TEXT("CPP_Default_")) + Param->GetName()) );
-				const FString MetadataCppDefaultValue = Function->GetMetaData(MetadataCppDefaultValueKey);
-				if(!MetadataCppDefaultValue.IsEmpty())
-				{
-					Pin->DefaultValue = Pin->AutogeneratedDefaultValue = MetadataCppDefaultValue;
-				}
-				else
-				{
-					// Set the default value to (T)0
-					K2Schema->SetPinDefaultValueBasedOnType(Pin);
-				}
-			}
+			K2Schema->SetPinDefaultValue(Pin, Function, Param);
 
 			// setup the default tool-tip text for this pin
 			GeneratePinHoverText(*Pin, Pin->PinToolTip);
@@ -789,8 +770,8 @@ FLinearColor UK2Node_CallFunction::GetNodeTitleColor() const
 		return GetDefault<UGraphEditorSettings>()->PureFunctionCallNodeTitleColor;
 	}
 
-	return Super::GetNodeTitleColor();
-}
+		return Super::GetNodeTitleColor();
+	}
 
 FString UK2Node_CallFunction::GetTooltip() const
 {
@@ -1222,6 +1203,8 @@ FNodeHandlingFunctor* UK2Node_CallFunction::CreateNodeHandler(FKismetCompilerCon
 
 void UK2Node_CallFunction::ExpandNode(class FKismetCompilerContext& CompilerContext, UEdGraph* SourceGraph)
 {
+	Super::ExpandNode(CompilerContext, SourceGraph);
+
 	if (CompilerContext.bIsFullCompile)
 	{
 		const UEdGraphSchema_K2* Schema = CompilerContext.GetSchema();
@@ -1234,29 +1217,29 @@ void UK2Node_CallFunction::ExpandNode(class FKismetCompilerContext& CompilerCont
 			TArray<UK2Node_FunctionEntry*> EntryPoints;
 			SourceGraph->GetNodesOfClass(EntryPoints);
 			if (1 != EntryPoints.Num())
-			{
+		{
 				CompilerContext.MessageLog.Warning(*FString::Printf(*LOCTEXT("WrongEntryPointsNum", "%i entry points found while expanding node @@").ToString(), EntryPoints.Num()), this);
-			}
+		}
 			else if (auto BetterSelfPin = EntryPoints[0]->GetAutoWorldContextPin())
 			{
 				FString const DefaultToSelfMetaValue = Function->GetMetaData(FBlueprintMetadata::MD_DefaultToSelf);
 				FString const WorldContextMetaValue = Function->GetMetaData(FBlueprintMetadata::MD_WorldContext);
 
 				struct FStructConnectHelper
-				{
+	{
 					static void Connect(const FString& PinName, UK2Node* Node, UEdGraphPin* BetterSelf, const UEdGraphSchema_K2* InSchema, FCompilerResultsLog& MessageLog)
-					{
+		{
 						auto Pin = Node->FindPin(PinName);
 						if (!PinName.IsEmpty() && Pin && !Pin->LinkedTo.Num())
-						{
+	{
 							const bool bConnected = InSchema->TryCreateConnection(Pin, BetterSelf);
 							if (!bConnected)
-							{
+		{
 								MessageLog.Warning(*LOCTEXT("DefaultToSelfNotConnected", "DefaultToSelf pin @@ from node @@ cannot be connected to @@").ToString(), Pin, Node, BetterSelf);
-							}
-						}
-					}
-				};
+		}
+		}
+	}
+};
 				FStructConnectHelper::Connect(DefaultToSelfMetaValue, this, BetterSelfPin, Schema, CompilerContext.MessageLog);
 				FStructConnectHelper::Connect(WorldContextMetaValue, this, BetterSelfPin, Schema, CompilerContext.MessageLog);
 			}
@@ -1276,50 +1259,50 @@ void UK2Node_CallFunction::ExpandNode(class FKismetCompilerContext& CompilerCont
 					// Expanded as input execs pins
 					if (EnumParamPin->Direction == EGPD_Input)
 					{
-						// Create normal exec input
-						UEdGraphPin* ExecutePin = CreatePin(EGPD_Input, Schema->PC_Exec, TEXT(""), NULL, false, false, Schema->PN_Execute);
-
-						// Create temp enum variable
-						UK2Node_TemporaryVariable* TempEnumVarNode = CompilerContext.SpawnIntermediateNode<UK2Node_TemporaryVariable>(this, SourceGraph);
-						TempEnumVarNode->VariableType.PinCategory = Schema->PC_Byte;
-						TempEnumVarNode->VariableType.PinSubCategoryObject = EnumProp->Enum;
-						TempEnumVarNode->AllocateDefaultPins();
-						// Get the output pin
-						UEdGraphPin* TempEnumVarOutput = TempEnumVarNode->GetVariablePin();
-
-						// Connect temp enum variable to (hidden) enum pin
-						Schema->TryCreateConnection(TempEnumVarOutput, EnumParamPin);
-
-						// Now we want to iterate over other exec inputs...
-						for(int32 PinIdx=Pins.Num()-1; PinIdx>=0; PinIdx--)
+					// Create normal exec input
+					UEdGraphPin* ExecutePin = CreatePin(EGPD_Input, Schema->PC_Exec, TEXT(""), NULL, false, false, Schema->PN_Execute);
+
+					// Create temp enum variable
+					UK2Node_TemporaryVariable* TempEnumVarNode = CompilerContext.SpawnIntermediateNode<UK2Node_TemporaryVariable>(this, SourceGraph);
+					TempEnumVarNode->VariableType.PinCategory = Schema->PC_Byte;
+					TempEnumVarNode->VariableType.PinSubCategoryObject = EnumProp->Enum;
+					TempEnumVarNode->AllocateDefaultPins();
+					// Get the output pin
+					UEdGraphPin* TempEnumVarOutput = TempEnumVarNode->GetVariablePin();
+
+					// Connect temp enum variable to (hidden) enum pin
+					Schema->TryCreateConnection(TempEnumVarOutput, EnumParamPin);
+
+					// Now we want to iterate over other exec inputs...
+					for(int32 PinIdx=Pins.Num()-1; PinIdx>=0; PinIdx--)
+					{
+						UEdGraphPin* Pin = Pins[PinIdx];
+						if( Pin != NULL && 
+							Pin != ExecutePin &&
+							Pin->Direction == EGPD_Input && 
+							Pin->PinType.PinCategory == Schema->PC_Exec )
 						{
-							UEdGraphPin* Pin = Pins[PinIdx];
-							if( Pin != NULL && 
-								Pin != ExecutePin &&
-								Pin->Direction == EGPD_Input && 
-								Pin->PinType.PinCategory == Schema->PC_Exec )
-							{
-								// Create node to set the temp enum var
-								UK2Node_AssignmentStatement* AssignNode = CompilerContext.SpawnIntermediateNode<UK2Node_AssignmentStatement>(this, SourceGraph);
-								AssignNode->AllocateDefaultPins();
+							// Create node to set the temp enum var
+							UK2Node_AssignmentStatement* AssignNode = CompilerContext.SpawnIntermediateNode<UK2Node_AssignmentStatement>(this, SourceGraph);
+							AssignNode->AllocateDefaultPins();
 
-								// Move connections from fake 'enum exec' pint to this assignment node
+							// Move connections from fake 'enum exec' pint to this assignment node
 								CompilerContext.MovePinLinksToIntermediate(*Pin, *AssignNode->GetExecPin());
 
-								// Connect this to out temp enum var
-								Schema->TryCreateConnection(AssignNode->GetVariablePin(), TempEnumVarOutput);
+							// Connect this to out temp enum var
+							Schema->TryCreateConnection(AssignNode->GetVariablePin(), TempEnumVarOutput);
 
-								// Connect exec output to 'real' exec pin
-								Schema->TryCreateConnection(AssignNode->GetThenPin(), ExecutePin);
+							// Connect exec output to 'real' exec pin
+							Schema->TryCreateConnection(AssignNode->GetThenPin(), ExecutePin);
 
-								// set the literal enum value to set to
-								AssignNode->GetValuePin()->DefaultValue = Pin->PinName;
+							// set the literal enum value to set to
+							AssignNode->GetValuePin()->DefaultValue = Pin->PinName;
 
-								// Finally remove this 'cosmetic' exec pin
-								Pins.RemoveAt(PinIdx);
-							}
+							// Finally remove this 'cosmetic' exec pin
+							Pins.RemoveAt(PinIdx);
 						}
 					}
+				}
 					// Expanded as output execs pins
 					else if (EnumParamPin->Direction == EGPD_Output)
 					{
diff --git a/Engine/Source/Editor/BlueprintGraph/Private/K2Node_ForEachElementInEnum.cpp b/Engine/Source/Editor/BlueprintGraph/Private/K2Node_ForEachElementInEnum.cpp
--- a/Engine/Source/Editor/BlueprintGraph/Private/K2Node_ForEachElementInEnum.cpp
+++ b/Engine/Source/Editor/BlueprintGraph/Private/K2Node_ForEachElementInEnum.cpp
@@ -150,6 +150,8 @@ FString UK2Node_ForEachElementInEnum::GetNodeNativeTitle(ENodeTitleType::Type Ti
 
 void UK2Node_ForEachElementInEnum::ExpandNode(class FKismetCompilerContext& CompilerContext, UEdGraph* SourceGraph)
 {
+	Super::ExpandNode(CompilerContext, SourceGraph);
+
 	if (CompilerContext.bIsFullCompile)
 	{
 		if (!Enum)
diff --git a/Engine/Source/Editor/BlueprintGraph/Private/K2Node_Knot.cpp b/Engine/Source/Editor/BlueprintGraph/Private/K2Node_Knot.cpp
--- a/Engine/Source/Editor/BlueprintGraph/Private/K2Node_Knot.cpp
+++ b/Engine/Source/Editor/BlueprintGraph/Private/K2Node_Knot.cpp
@@ -65,6 +65,8 @@ void UK2Node_Knot::GetMenuEntries(struct FGraphContextMenuBuilder& ContextMenuBu
 
 void UK2Node_Knot::ExpandNode(class FKismetCompilerContext& CompilerContext, UEdGraph* SourceGraph)
 {
+	Super::ExpandNode(CompilerContext, SourceGraph);
+
 	const UEdGraphSchema_K2* K2Schema = GetDefault<UEdGraphSchema_K2>();
 
 	UEdGraphPin* MyInputPin = GetInputPin();
diff --git a/Engine/Source/Editor/BlueprintGraph/Private/K2Node_MacroInstance.cpp b/Engine/Source/Editor/BlueprintGraph/Private/K2Node_MacroInstance.cpp
--- a/Engine/Source/Editor/BlueprintGraph/Private/K2Node_MacroInstance.cpp
+++ b/Engine/Source/Editor/BlueprintGraph/Private/K2Node_MacroInstance.cpp
@@ -81,15 +81,19 @@ void UK2Node_MacroInstance::AllocateDefaultPins()
 					{
 						UEdGraphPin* PortPin = *PinIt;
 
-						UEdGraphPin* NewLocalPin = CreatePin(
-							UEdGraphPin::GetComplementaryDirection(PortPin->Direction),
-							PortPin->PinType.PinCategory,
-							PortPin->PinType.PinSubCategory,
-							PortPin->PinType.PinSubCategoryObject.Get(),
-							PortPin->PinType.bIsArray,
-							PortPin->PinType.bIsReference,
-							PortPin->PinName);
-						NewLocalPin->DefaultValue = NewLocalPin->AutogeneratedDefaultValue = PortPin->DefaultValue;
+						// We're not interested in any pins that have been expanded internally on the macro
+						if (PortPin->ParentPin == NULL)
+						{
+							UEdGraphPin* NewLocalPin = CreatePin(
+								UEdGraphPin::GetComplementaryDirection(PortPin->Direction),
+								PortPin->PinType.PinCategory,
+								PortPin->PinType.PinSubCategory,
+								PortPin->PinType.PinSubCategoryObject.Get(),
+								PortPin->PinType.bIsArray,
+								PortPin->PinType.bIsReference,
+								PortPin->PinName);
+							NewLocalPin->DefaultValue = NewLocalPin->AutogeneratedDefaultValue = PortPin->DefaultValue;
+						}
 					}
 				}
 			}
diff --git a/Engine/Source/Editor/BlueprintGraph/Private/K2Node_MakeStruct.cpp b/Engine/Source/Editor/BlueprintGraph/Private/K2Node_MakeStruct.cpp
--- a/Engine/Source/Editor/BlueprintGraph/Private/K2Node_MakeStruct.cpp
+++ b/Engine/Source/Editor/BlueprintGraph/Private/K2Node_MakeStruct.cpp
@@ -136,48 +136,48 @@ void UK2Node_MakeStruct::FMakeStructPinManager::CustomizePinData(UEdGraphPin* Pi
 				InPin.DefaultTextValue = FText::FromString(Value);
 			}
 			else
-			{
+	{
 				InPin.DefaultValue = Value;
 			}
-		}
+	}
 	};
 
 	UK2Node_StructOperation::FStructOperationOptionalPinManager::CustomizePinData(Pin, SourcePropertyName, ArrayIndex, Property);
-	if (Pin && Property)
-	{
-		const UEdGraphSchema_K2* Schema = GetDefault<UEdGraphSchema_K2>();
-		check(Schema);
+		if(Pin && Property)
+		{
+			const UEdGraphSchema_K2* Schema = GetDefault<UEdGraphSchema_K2>();
+			check(Schema);
 		const bool bIsText = Property->IsA<UTextProperty>();
 		checkSlow(bIsText == (Schema->PC_Text == Pin->PinType.PinCategory));
 
-		const FString& MetadataDefaultValue = Property->GetMetaData(TEXT("MakeStructureDefaultValue"));
-		if (!MetadataDefaultValue.IsEmpty())
-		{
+			const FString& MetadataDefaultValue = Property->GetMetaData(TEXT("MakeStructureDefaultValue"));
+			if (!MetadataDefaultValue.IsEmpty())
+			{
 			FPinDefaultValueHelper::Set(*Pin, MetadataDefaultValue, bIsText);
-			return;
-		}
+				return;
+			}
 
-		if (NULL != SampleStructMemory)
-		{
-			FString NewDefaultValue;
-			if (Property->ExportText_InContainer(0, NewDefaultValue, SampleStructMemory, SampleStructMemory, NULL, PPF_None))
+			if(NULL != SampleStructMemory)
 			{
-				if (Schema->IsPinDefaultValid(Pin, NewDefaultValue, NULL, FText::GetEmpty()).IsEmpty())
+				FString NewDefaultValue;
+				if(Property->ExportText_InContainer(0, NewDefaultValue, SampleStructMemory, SampleStructMemory, NULL, PPF_None))
 				{
+					if(Schema->IsPinDefaultValid(Pin, NewDefaultValue, NULL, FText::GetEmpty()).IsEmpty())
+					{
 					FPinDefaultValueHelper::Set(*Pin, NewDefaultValue, bIsText);
-					return;
+						return;
+					}
 				}
 			}
+			
+			Schema->SetPinDefaultValueBasedOnType(Pin);
 		}
-
-		Schema->SetPinDefaultValueBasedOnType(Pin);
 	}
-}
 
 bool UK2Node_MakeStruct::FMakeStructPinManager::CanTreatPropertyAsOptional(UProperty* TestProperty) const
-{
-	return UK2Node_MakeStruct::CanBeExposed(TestProperty);
-}
+	{
+		return UK2Node_MakeStruct::CanBeExposed(TestProperty);
+	}
 
 UK2Node_MakeStruct::UK2Node_MakeStruct(const class FPostConstructInitializeProperties& PCIP)
 	: Super(PCIP)
@@ -281,7 +281,7 @@ bool UK2Node_MakeStruct::CanBeExposed(const UProperty* Property)
 
 bool UK2Node_MakeStruct::CanBeMade(const UScriptStruct* Struct)
 {
-	if(!Struct->GetBoolMetaData(TEXT("HasNativeMakeBreak")))
+	if(!Struct->HasMetaData(TEXT("HasNativeMake")))
 	{
 		for (TFieldIterator<UProperty> It(Struct); It; ++It)
 		{
@@ -302,9 +302,29 @@ FNodeHandlingFunctor* UK2Node_MakeStruct::CreateNodeHandler(FKismetCompilerConte
 UK2Node::ERedirectType UK2Node_MakeStruct::DoPinsMatchForReconstruction(const UEdGraphPin* NewPin, int32 NewPinIndex, const UEdGraphPin* OldPin, int32 OldPinIndex)  const
 {
 	ERedirectType Result = UK2Node::DoPinsMatchForReconstruction(NewPin, NewPinIndex, OldPin, OldPinIndex);
-	if ((ERedirectType_None == Result) && DoRenamedPinsMatch(NewPin, OldPin, false))
+	if ((ERedirectType_None == Result) && NewPin && OldPin)
+	{
+		const UEdGraphSchema_K2* K2Schema = GetDefault<UEdGraphSchema_K2>();
+		if ((EGPD_Output == NewPin->Direction) && (EGPD_Output == OldPin->Direction))
 	{
-		Result = ERedirectType_Custom;
+		const UEdGraphSchema_K2* K2Schema = GetDefault<UEdGraphSchema_K2>();
+		if (K2Schema->ArePinTypesCompatible( NewPin->PinType, OldPin->PinType))
+		{
+			Result = ERedirectType_Custom;
+		}
+	}
+		else if ((EGPD_Input == NewPin->Direction) && (EGPD_Input == OldPin->Direction))
+		{
+			TMap<FName, FName>* StructRedirects = UStruct::TaggedPropertyRedirects.Find(StructType->GetFName());
+			if (StructRedirects)
+			{
+				FName* PropertyRedirect = StructRedirects->Find(FName(*OldPin->PinName));
+				if (PropertyRedirect)
+				{
+					Result = ((FCString::Stricmp(*PropertyRedirect->ToString(), *NewPin->PinName) != 0) ? ERedirectType_None : ERedirectType_Name);
+				}
+			}
+		}
 	}
 	return Result;
 }
diff --git a/Engine/Source/Editor/GameplayTagsEditor/Private/GameplayTagsK2Node_LiteralGameplayTag.cpp b/Engine/Source/Editor/GameplayTagsEditor/Private/GameplayTagsK2Node_LiteralGameplayTag.cpp
--- a/Engine/Source/Editor/GameplayTagsEditor/Private/GameplayTagsK2Node_LiteralGameplayTag.cpp
+++ b/Engine/Source/Editor/GameplayTagsEditor/Private/GameplayTagsK2Node_LiteralGameplayTag.cpp
@@ -54,6 +54,8 @@ void UGameplayTagsK2Node_LiteralGameplayTag::GetMenuEntries( FGraphContextMenuBu
 
 void UGameplayTagsK2Node_LiteralGameplayTag::ExpandNode(class FKismetCompilerContext& CompilerContext, UEdGraph* SourceGraph)
 {
+	Super::ExpandNode(CompilerContext, SourceGraph);
+
 	if (CompilerContext.bIsFullCompile)
 	{
 		const UEdGraphSchema_K2* Schema = CompilerContext.GetSchema();
diff --git a/Engine/Source/Editor/GraphEditor/Private/GraphEditorActions.cpp b/Engine/Source/Editor/GraphEditor/Private/GraphEditorActions.cpp
--- a/Engine/Source/Editor/GraphEditor/Private/GraphEditorActions.cpp
+++ b/Engine/Source/Editor/GraphEditor/Private/GraphEditorActions.cpp
@@ -45,6 +45,8 @@ void FGraphEditorCommandsImpl::RegisterCommands()
 
 	UI_COMMAND( BreakPinLinks, "Break Link(s)", "Breaks pin links", EUserInterfaceActionType::Button, FInputGesture() )
 	UI_COMMAND( PromoteToVariable, "Promote to Variable", "Promotes something to a variable", EUserInterfaceActionType::Button, FInputGesture() )
+	UI_COMMAND( SplitStructPin, "Split Struct Pin", "Breaks a struct pin in to a separate pin per element", EUserInterfaceActionType::Button, FInputGesture() )
+	UI_COMMAND( RecombineStructPin, "Recombine Struct Pin", "Takes struct pins that have been broken in to composite elements and combines them back to a single struct pin", EUserInterfaceActionType::Button, FInputGesture() )
 	UI_COMMAND( StartWatchingPin, "Watch this value", "Adds this pin or variable to the watch list", EUserInterfaceActionType::Button, FInputGesture() )
 	UI_COMMAND( StopWatchingPin, "Stop watching this value", "Removes this pin or variable from the watch list ", EUserInterfaceActionType::Button, FInputGesture() )
 
diff --git a/Engine/Source/Editor/GraphEditor/Public/GraphEditorActions.h b/Engine/Source/Editor/GraphEditor/Public/GraphEditorActions.h
--- a/Engine/Source/Editor/GraphEditor/Public/GraphEditorActions.h
+++ b/Engine/Source/Editor/GraphEditor/Public/GraphEditorActions.h
@@ -68,6 +68,8 @@ public:
 	// Pin-specific actions
 	TSharedPtr< FUICommandInfo > BreakPinLinks;
 	TSharedPtr< FUICommandInfo > PromoteToVariable;
+	TSharedPtr< FUICommandInfo > SplitStructPin;
+	TSharedPtr< FUICommandInfo > RecombineStructPin;
 	TSharedPtr< FUICommandInfo > StartWatchingPin;
 	TSharedPtr< FUICommandInfo > StopWatchingPin;
 
diff --git a/Engine/Source/Editor/Kismet/Private/BlueprintEditor.cpp b/Engine/Source/Editor/Kismet/Private/BlueprintEditor.cpp
--- a/Engine/Source/Editor/Kismet/Private/BlueprintEditor.cpp
+++ b/Engine/Source/Editor/Kismet/Private/BlueprintEditor.cpp
@@ -442,6 +442,16 @@ TSharedRef<SGraphEditor> FBlueprintEditor::CreateGraphEditorWidget(TSharedRef<FT
 				FCanExecuteAction::CreateSP( this, &FBlueprintEditor::CanPromoteToVariable )
 				);
 
+			GraphEditorCommands->MapAction(FGraphEditorCommands::Get().SplitStructPin,
+				FExecuteAction::CreateSP( this, &FBlueprintEditor::OnSplitStructPin ),
+				FCanExecuteAction::CreateSP( this, &FBlueprintEditor::CanSplitStructPin )
+				);
+
+			GraphEditorCommands->MapAction(FGraphEditorCommands::Get().RecombineStructPin,
+				FExecuteAction::CreateSP( this, &FBlueprintEditor::OnRecombineStructPin ),
+				FCanExecuteAction::CreateSP( this, &FBlueprintEditor::CanRecombineStructPin )
+				);
+
 			GraphEditorCommands->MapAction( FGraphEditorCommands::Get().AddExecutionPin,
 				FExecuteAction::CreateSP( this, &FBlueprintEditor::OnAddExecutionPin ),
 				FCanExecuteAction::CreateSP( this, &FBlueprintEditor::CanAddExecutionPin )
@@ -2691,6 +2701,72 @@ bool FBlueprintEditor::CanPromoteToVariable() const
 	return bCanPromote;
 }
 
+void FBlueprintEditor::OnSplitStructPin()
+{
+	TSharedPtr<SGraphEditor> FocusedGraphEd = FocusedGraphEdPtr.Pin();
+	if (FocusedGraphEd.IsValid())
+	{
+		UEdGraphPin* TargetPin = FocusedGraphEd->GetGraphPinForMenu();
+
+		check(IsEditingSingleBlueprint());
+		check(GetBlueprintObj()->SkeletonGeneratedClass);
+		check(TargetPin);
+
+		const UEdGraphSchema_K2* K2Schema = GetDefault<UEdGraphSchema_K2>();
+		K2Schema->SplitPin(TargetPin);
+	}
+}
+
+bool FBlueprintEditor::CanSplitStructPin() const
+{
+	const UEdGraphSchema_K2* K2Schema = GetDefault<UEdGraphSchema_K2>();
+
+	bool bCanSplit = false;
+	TSharedPtr<SGraphEditor> FocusedGraphEd = FocusedGraphEdPtr.Pin();
+	if (FocusedGraphEd.IsValid())
+	{
+		if (UEdGraphPin* Pin = FocusedGraphEd->GetGraphPinForMenu())
+		{
+			bCanSplit = K2Schema->CanSplitStructPin(*Pin);
+		}
+	}
+
+	return bCanSplit;
+}
+
+void FBlueprintEditor::OnRecombineStructPin()
+{
+	TSharedPtr<SGraphEditor> FocusedGraphEd = FocusedGraphEdPtr.Pin();
+	if (FocusedGraphEd.IsValid())
+	{
+		UEdGraphPin* TargetPin = FocusedGraphEd->GetGraphPinForMenu();
+
+		check(IsEditingSingleBlueprint());
+		check(GetBlueprintObj()->SkeletonGeneratedClass);
+		check(TargetPin);
+
+		const UEdGraphSchema_K2* K2Schema = GetDefault<UEdGraphSchema_K2>();
+		K2Schema->RecombinePin(TargetPin);
+	}
+}
+
+bool FBlueprintEditor::CanRecombineStructPin() const
+{
+	const UEdGraphSchema_K2* K2Schema = GetDefault<UEdGraphSchema_K2>();
+
+	bool bCanRecombine = false;
+	TSharedPtr<SGraphEditor> FocusedGraphEd = FocusedGraphEdPtr.Pin();
+	if (FocusedGraphEd.IsValid())
+	{
+		if (UEdGraphPin* Pin = FocusedGraphEd->GetGraphPinForMenu())
+		{
+			bCanRecombine = K2Schema->CanRecombineStructPin(*Pin);
+		}
+	}
+
+	return bCanRecombine;
+}
+
 void FBlueprintEditor::OnAddExecutionPin()
 {
 	const FGraphPanelSelectionSet& SelectedNodes = GetSelectedNodes();
diff --git a/Engine/Source/Editor/Kismet/Public/BlueprintEditor.h b/Engine/Source/Editor/Kismet/Public/BlueprintEditor.h
--- a/Engine/Source/Editor/Kismet/Public/BlueprintEditor.h
+++ b/Engine/Source/Editor/Kismet/Public/BlueprintEditor.h
@@ -567,6 +567,12 @@ protected:
 	void OnPromoteToVariable();
 	bool CanPromoteToVariable() const;
 
+	void OnSplitStructPin();
+	bool CanSplitStructPin() const;
+
+	void OnRecombineStructPin();
+	bool CanRecombineStructPin() const;
+
 	void OnAddExecutionPin();
 	bool CanAddExecutionPin() const;
 
diff --git a/Engine/Source/Editor/KismetCompiler/Private/KismetCompiler.cpp b/Engine/Source/Editor/KismetCompiler/Private/KismetCompiler.cpp
--- a/Engine/Source/Editor/KismetCompiler/Private/KismetCompiler.cpp
+++ b/Engine/Source/Editor/KismetCompiler/Private/KismetCompiler.cpp
@@ -163,7 +163,7 @@ void FKismetCompilerContext::CleanAndSanitizeClass(UBlueprintGeneratedClass* Cla
 		ParentClass = UObject::StaticClass();
 	}
 	TransientClass->ClassAddReferencedObjects = ParentClass->AddReferencedObjects;
-	
+    
 	NewClass = ClassToClean;
 	OldCDO = ClassToClean->ClassDefaultObject; // we don't need to create the CDO at this point
 	
@@ -187,31 +187,31 @@ void FKismetCompilerContext::CleanAndSanitizeClass(UBlueprintGeneratedClass* Cla
 		{
 			SubObjectsToSave.AddObject(Blueprint->SimpleConstructionScript);
 			if (const USCS_Node* DefaultScene = Blueprint->SimpleConstructionScript->GetDefaultSceneRootNode())
-			{
+	{
 				SubObjectsToSave.AddObject(DefaultScene->ComponentTemplate);
-			}
+	}
 
 			TArray<USCS_Node*> SCSNodes = Blueprint->SimpleConstructionScript->GetAllNodes();
 			for (auto SCSNode : SCSNodes)
-			{
+	{
 				SubObjectsToSave.AddObject(SCSNode->ComponentTemplate);
 			}
-		}
-		{
+	}
+	{
 			TSet<class UCurveBase*> Curves;
 			for (auto Timeline : Blueprint->Timelines)
 			{
 				if (Timeline)
-				{
+		{
 					Timeline->GetAllCurves(Curves);
-				}
+		}
 			}
 			for (auto Component : Blueprint->ComponentTemplates)
 			{
 				if (auto TimelineComponent = Cast<UTimelineComponent>(Component))
-				{
+		{
 					TimelineComponent->GetAllCurves(Curves);
-				}
+		}
 			}
 			for (auto Curve : Curves)
 			{
@@ -294,11 +294,11 @@ void FKismetCompilerContext::ValidatePin(const UEdGraphPin* Pin) const
 		UEdGraphPin* MutablePin = const_cast<UEdGraphPin*>(Pin);
 
 		if (NewClass && MutablePin->PinType.PinSubCategoryObject.Get() == NewClass)
-		{
-			MutablePin->PinType.PinSubCategory = Schema->PSC_Self;
-			MutablePin->PinType.PinSubCategoryObject = NULL;
+			{
+				MutablePin->PinType.PinSubCategory = Schema->PSC_Self;
+				MutablePin->PinType.PinSubCategoryObject = NULL;
+			}
 		}
-	}
 
 	if (Pin->PinType.PinCategory == Schema->PC_Wildcard)
 	{
@@ -1010,7 +1010,7 @@ void FKismetCompilerContext::ValidateSelfPinsInGraph(const UEdGraph* SourceGraph
 					if (Schema->IsSelfPin(*Pin) && (Pin->LinkedTo.Num() == 0))
 					{
 						FEdGraphPinType SelfType;
-						SelfType.PinCategory    = Schema->PC_Object;
+						SelfType.PinCategory = Schema->PC_Object;
 						SelfType.PinSubCategory = Schema->PSC_Self;
 
 						if (!Schema->ArePinTypesCompatible(SelfType, Pin->PinType, NewClass))
@@ -1105,7 +1105,7 @@ void FKismetCompilerContext::PrecompileFunction(FKismetFunctionContext& Context)
 		}
 
 		// Determine if this is a new function or if it overrides a parent function
-		//@TODO: Does not support multiple overloads for a parent virtual function
+			//@TODO: Does not support multiple overloads for a parent virtual function
 		UFunction* ParentFunction = Blueprint->ParentClass->FindFunctionByName(NewFunctionName);
 
 		const FString NewFunctionNameString = NewFunctionName.ToString();
@@ -1540,7 +1540,7 @@ void FKismetCompilerContext::FinishCompilingFunction(FKismetFunctionContext& Con
 		// No defaults for object/class pins
 		if(	!Schema->IsMetaPin(*EntryPin) && 
 			(EntryPin->PinType.PinCategory != Schema->PC_Object) && 
-			(EntryPin->PinType.PinCategory != Schema->PC_Class)  && 
+			(EntryPin->PinType.PinCategory != Schema->PC_Class) && 
 			(EntryPin->PinType.PinCategory != Schema->PC_Interface) && 
 			!EntryPin->DefaultValue.IsEmpty() )
 		{
@@ -2190,20 +2190,20 @@ void FKismetCompilerContext::CreateFunctionStubForEvent(UK2Node_Event* SrcEventN
 
 	// If this is a customizable event, make sure to copy over the user defined pins
 	if (UK2Node_CustomEvent const* SrcCustomEventNode = Cast<UK2Node_CustomEvent const>(SrcEventNode))
-	{
+ 	{
 		EntryNode->UserDefinedPins = SrcCustomEventNode->UserDefinedPins;
 		// CustomEvents may inherit net flags (so let's use their GetNetFlags() incase this is an override)
 		StubContext.MarkAsNetFunction(SrcCustomEventNode->GetNetFlags());
-	}
+ 	}
 	EntryNode->AllocateDefaultPins();
 
 	// Confirm that the event node matches the latest function signature, which the newly created EntryNode should have
-	if( !SrcEventNode->IsFunctionEntryCompatible(EntryNode) )
-	{
-		// There is no match, so the function parameters must have changed.  Throw an error, and force them to refresh
-		MessageLog.Error(*FString::Printf(*LOCTEXT("EventNodeOutOfDate_Error", "Event node @@ is out-of-date.  Please refresh it.").ToString()), SrcEventNode);
-		return;
-	}
+ 	if( !SrcEventNode->IsFunctionEntryCompatible(EntryNode) )
+ 	{
+ 		// There is no match, so the function parameters must have changed.  Throw an error, and force them to refresh
+ 		MessageLog.Error(*FString::Printf(*LOCTEXT("EventNodeOutOfDate_Error", "Event node @@ is out-of-date.  Please refresh it.").ToString()), SrcEventNode);
+ 		return;
+ 	}
 
 	// Copy each event parameter to the assignment node, if there are any inputs
 	UK2Node_VariableSet* AssignmentNode = NULL;
@@ -2248,7 +2248,7 @@ void FKismetCompilerContext::CreateFunctionStubForEvent(UK2Node_Event* SrcEventN
 	CallIntoUbergraph->NodePosX = 300;
 
 	// Use the ExecuteUbergraph base function to generate the pins...
-	CallIntoUbergraph->FunctionReference.SetExternalMember(Schema->FN_ExecuteUbergraphBase, UObject::StaticClass());
+ 	CallIntoUbergraph->FunctionReference.SetExternalMember(Schema->FN_ExecuteUbergraphBase, UObject::StaticClass());
 	CallIntoUbergraph->AllocateDefaultPins();
 	
 	// ...then swap to the generated version for this level
@@ -2321,30 +2321,30 @@ void FKismetCompilerContext::ExpansionStep(UEdGraph* Graph, bool bAllowUbergraph
 {
 	if (bIsFullCompile)
 	{
-		SCOPE_CYCLE_COUNTER(EKismetCompilerStats_Expansion);
+	SCOPE_CYCLE_COUNTER(EKismetCompilerStats_Expansion);
 
-		// Collapse any remaining tunnels or macros
-		ExpandTunnelsAndMacros(Graph);
+	// Collapse any remaining tunnels or macros
+	ExpandTunnelsAndMacros(Graph);
 
-		for (int32 NodeIndex = 0; NodeIndex < Graph->Nodes.Num(); ++NodeIndex)
+	for (int32 NodeIndex = 0; NodeIndex < Graph->Nodes.Num(); ++NodeIndex)
+	{
+		UK2Node* Node = Cast<UK2Node>(Graph->Nodes[NodeIndex]);
+		if (Node)
 		{
-			UK2Node* Node = Cast<UK2Node>(Graph->Nodes[NodeIndex]);
-			if (Node)
-			{
-				Node->ExpandNode(*this, Graph);
-			}
+			Node->ExpandNode(*this, Graph);
 		}
+	}
 
-		if (bAllowUbergraphExpansions)
-		{
-			// Expand timeline nodes
-			ExpandTimelineNodes(Graph);
+	if (bAllowUbergraphExpansions)
+	{
+		// Expand timeline nodes
+		ExpandTimelineNodes(Graph);
 
-			// Expand PlayMovieScene nodes
-			ExpandPlayMovieSceneNodes(Graph);
-		}
+		// Expand PlayMovieScene nodes
+		ExpandPlayMovieSceneNodes(Graph);
 	}
 }
+}
 
 void FKismetCompilerContext::VerifyValidOverrideEvent(const UEdGraph* Graph)
 {
@@ -2621,10 +2621,10 @@ void FKismetCompilerContext::ExpandTunnelsAndMacros(UEdGraph* SourceGraph)
 								ClonedPin->PinType.PinSubCategory = MacroInstanceNode->ResolvedWildcardType.PinSubCategory;
 								ClonedPin->PinType.PinSubCategoryObject = MacroInstanceNode->ResolvedWildcardType.PinSubCategoryObject;
 							}
+							}
 						}
 					}
 				}
-			}
 
 			// Handle any nodes that need to inherit their macro instance's NodeGUID
 			for( auto ClonedNode : MacroNodes)
@@ -2636,21 +2636,30 @@ void FKismetCompilerContext::ExpandTunnelsAndMacros(UEdGraph* SourceGraph)
 				}
 			}
 
-			// Since we don't support array literals, drop a make array node on any unconnected array pins, which will allow macro expansion to succeed even if disconnected
-			for (auto PinIt = MacroInstanceNode->Pins.CreateIterator(); PinIt; ++PinIt)
+			// We iterate the array in reverse so we can both remove the subpins safely after we've read them and
+			// so we have split nested structs we combine them back together in the right order
+			for (int32 PinIndex = MacroInstanceNode->Pins.Num() - 1; PinIndex >= 0; --PinIndex)
 			{
-				UEdGraphPin* Pin = *PinIt;
-				if (Pin 
-					&& Pin->PinType.bIsArray
+				UEdGraphPin* Pin = MacroInstanceNode->Pins[PinIndex];
+				if (Pin)
+				{
+					// Since we don't support array literals, drop a make array node on any unconnected array pins, which will allow macro expansion to succeed even if disconnected
+					if (Pin->PinType.bIsArray
 					&& (Pin->Direction == EGPD_Input)
 					&& (Pin->LinkedTo.Num() == 0))
-				{
-					UK2Node_MakeArray* MakeArrayNode = SpawnIntermediateNode<UK2Node_MakeArray>(MacroInstanceNode, SourceGraph);
-					MakeArrayNode->AllocateDefaultPins();
-					UEdGraphPin* MakeArrayOut = MakeArrayNode->GetOutputPin();
-					check(MakeArrayOut);
-					MakeArrayOut->MakeLinkTo(Pin);
-					MakeArrayNode->PinConnectionListChanged(MakeArrayOut);
+					{
+						UK2Node_MakeArray* MakeArrayNode = SpawnIntermediateNode<UK2Node_MakeArray>(MacroInstanceNode, SourceGraph);
+						MakeArrayNode->AllocateDefaultPins();
+						UEdGraphPin* MakeArrayOut = MakeArrayNode->GetOutputPin();
+						check(MakeArrayOut);
+						MakeArrayOut->MakeLinkTo(Pin);
+						MakeArrayNode->PinConnectionListChanged(MakeArrayOut);
+					}
+					// Otherwise we need to handle the pin splitting
+					else if (Pin->SubPins.Num() > 0)
+					{
+						MacroInstanceNode->ExpandSplitPin(this, SourceGraph, Pin);
+					}
 				}
 			}
 
@@ -2753,7 +2762,7 @@ void FKismetCompilerContext::ExpandTunnelsAndMacros(UEdGraph* SourceGraph)
 				}
 			}
 
-			bool bSuccess = Schema->CollapseGatewayNode(TunnelNode, InputSink, TunnelNode->GetOutputSource());
+			bool bSuccess = Schema->CollapseGatewayNode(TunnelNode, InputSink, TunnelNode->GetOutputSource(), this);
 			if (!bSuccess)
 			{
 				MessageLog.Error(*LOCTEXT("CollapseTunnel_Error", "Failed to collapse tunnel @@").ToString(), TunnelNode);
@@ -2931,25 +2940,25 @@ void FKismetCompilerContext::Compile()
 	UBlueprintGeneratedClass* TargetClass = Cast<UBlueprintGeneratedClass>(TargetUClass);
 
 	// >>> Backwards Compatibility: Make sure that skeleton generated classes have the proper "SKEL_" naming convention
-	const FString SkeletonPrefix(TEXT("SKEL_"));
-	if( bIsSkeletonOnly && TargetClass && !TargetClass->GetName().StartsWith(SkeletonPrefix) )
-	{
-		FString NewName = SkeletonPrefix + TargetClass->GetName();
+ 	const FString SkeletonPrefix(TEXT("SKEL_"));
+ 	if( bIsSkeletonOnly && TargetClass && !TargetClass->GetName().StartsWith(SkeletonPrefix) )
+ 	{
+ 		FString NewName = SkeletonPrefix + TargetClass->GetName();
  
-		// Ensure we have a free name for this class
-		UClass* AnyClassWithGoodName = (UClass*)StaticFindObject(UClass::StaticClass(), Blueprint->GetOutermost(), *NewName, false);
-		if( AnyClassWithGoodName )
-		{
-			// Special Case:  If the CDO of the class has become dissociated from its actual CDO, attempt to find the proper named CDO, and get rid of it.
-			if( AnyClassWithGoodName->ClassDefaultObject == TargetClass->ClassDefaultObject )
-			{
-				AnyClassWithGoodName->ClassDefaultObject = NULL;
-				FString DefaultObjectName = FString(DEFAULT_OBJECT_PREFIX) + NewName;
-				AnyClassWithGoodName->ClassDefaultObject = (UObject*)StaticFindObject(UObject::StaticClass(), Blueprint->GetOutermost(), *DefaultObjectName, false);
-			}
+ 		// Ensure we have a free name for this class
+ 		UClass* AnyClassWithGoodName = (UClass*)StaticFindObject(UClass::StaticClass(), Blueprint->GetOutermost(), *NewName, false);
+ 		if( AnyClassWithGoodName )
+ 		{
+ 			// Special Case:  If the CDO of the class has become dissociated from its actual CDO, attempt to find the proper named CDO, and get rid of it.
+ 			if( AnyClassWithGoodName->ClassDefaultObject == TargetClass->ClassDefaultObject )
+ 			{
+ 				AnyClassWithGoodName->ClassDefaultObject = NULL;
+ 				FString DefaultObjectName = FString(DEFAULT_OBJECT_PREFIX) + NewName;
+ 				AnyClassWithGoodName->ClassDefaultObject = (UObject*)StaticFindObject(UObject::StaticClass(), Blueprint->GetOutermost(), *DefaultObjectName, false);
+ 			}
  
 			// Get rid of the old class to make room for renaming our class to the final SKEL name
-			FKismetCompilerUtilities::ConsignToOblivion(AnyClassWithGoodName, Blueprint->bIsRegeneratingOnLoad);
+ 			FKismetCompilerUtilities::ConsignToOblivion(AnyClassWithGoodName, Blueprint->bIsRegeneratingOnLoad);
 
 			// Update the refs to the old SKC
 			TMap<UObject*, UObject*> ClassReplacementMap;
@@ -2960,11 +2969,11 @@ void FKismetCompilerContext::Compile()
 			{
 				FArchiveReplaceObjectRef<UObject> ReplaceInBlueprintAr(AllGraphs[i], ClassReplacementMap, /*bNullPrivateRefs=*/ false, /*bIgnoreOuterRef=*/ false, /*bIgnoreArchetypeRef=*/ false);
 			}
-		}
+ 		}
  
-		ERenameFlags RenameFlags = (REN_DontCreateRedirectors|REN_NonTransactional|(Blueprint->bIsRegeneratingOnLoad ? REN_ForceNoResetLoaders : 0));
-		TargetClass->Rename(*NewName, NULL, RenameFlags);
-	}
+ 		ERenameFlags RenameFlags = (REN_DontCreateRedirectors|REN_NonTransactional|(Blueprint->bIsRegeneratingOnLoad ? REN_ForceNoResetLoaders : 0));
+ 		TargetClass->Rename(*NewName, NULL, RenameFlags);
+ 	}
 	// <<< End Backwards Compatibility
 
 	// >>> Backwards compatibility:  If SkeletonGeneratedClass == GeneratedClass, we need to make a new generated class the first time we need it
@@ -3263,7 +3272,7 @@ void FKismetCompilerContext::Compile()
 						}
 					}
 
-					UEditorEngine::CopyPropertiesForUnrelatedObjects(OldCDO, NewCDO);
+					UEditorEngine::CopyPropertiesForUnrelatedObjects(OldCDO, NewCDO);			
 				}
 
 				// >>> Backwards Compatibility: Propagate data from the skel CDO to the gen CDO if we haven't already done so for this blueprint
@@ -3358,29 +3367,29 @@ void FKismetCompilerContext::Compile()
 	}
 
 	// For full compiles, find other blueprints that may need refreshing, and mark them dirty, in case they try to run
-	if( bIsFullCompile && !Blueprint->bIsRegeneratingOnLoad )
-	{
-		TArray<UObject*> AllBlueprints;
-		GetObjectsOfClass(UBlueprint::StaticClass(), AllBlueprints, true);
+  	if( bIsFullCompile && !Blueprint->bIsRegeneratingOnLoad )
+  	{
+  		TArray<UObject*> AllBlueprints;
+  		GetObjectsOfClass(UBlueprint::StaticClass(), AllBlueprints, true);
   
-		// Mark any blueprints that implement this interface as dirty
-		for( auto CurrentObj = AllBlueprints.CreateIterator(); CurrentObj; ++CurrentObj )
-		{
-			UBlueprint* CurrentBP = Cast<UBlueprint>(*CurrentObj);
+  		// Mark any blueprints that implement this interface as dirty
+  		for( auto CurrentObj = AllBlueprints.CreateIterator(); CurrentObj; ++CurrentObj )
+  		{
+  			UBlueprint* CurrentBP = Cast<UBlueprint>(*CurrentObj);
   
-			if( FBlueprintEditorUtils::IsBlueprintDependentOn(CurrentBP, Blueprint) )
-			{
+  			if( FBlueprintEditorUtils::IsBlueprintDependentOn(CurrentBP, Blueprint) )
+  			{
 				CurrentBP->Status = BS_Dirty;
-				FBlueprintEditorUtils::RefreshExternalBlueprintDependencyNodes(CurrentBP);
+  				FBlueprintEditorUtils::RefreshExternalBlueprintDependencyNodes(CurrentBP);
 				CurrentBP->BroadcastChanged();
-			}
-		}
-	}
+  			}
+  		}
+  	}
 
 	// Clear out pseudo-local members that are only valid within a Compile call
 	UbergraphContext = NULL;
-	CallsIntoUbergraph.Empty();
-	TimelineToMemberVariableMap.Empty();
+ 	CallsIntoUbergraph.Empty();
+ 	TimelineToMemberVariableMap.Empty();
 
 
 	check(NewClass->PropertiesSize >= UObject::StaticClass()->PropertiesSize);
diff --git a/Engine/Source/Runtime/CoreUObject/Classes/Object.h b/Engine/Source/Runtime/CoreUObject/Classes/Object.h
--- a/Engine/Source/Runtime/CoreUObject/Classes/Object.h
+++ b/Engine/Source/Runtime/CoreUObject/Classes/Object.h
@@ -182,7 +182,7 @@ struct FNetworkGUID
 };
 
 // A point or direction FVector in 3d space.
-USTRUCT(immutable, noexport, BlueprintType, meta=(HasNativeMakeBreak="true"))
+USTRUCT(immutable, noexport, BlueprintType, meta=(HasNativeMake="Engine.KismetMathLibrary.MakeVector", HasNativeBreak="Engine.KismetMathLibrary.BreakVector"))
 struct FVector
 {
 	UPROPERTY(EditAnywhere, BlueprintReadWrite, Category=Vector, SaveGame)
@@ -214,7 +214,7 @@ struct FVector4
 };
 
 
-USTRUCT(immutable, noexport, BlueprintType, meta=(HasNativeMakeBreak="true"))
+USTRUCT(immutable, noexport, BlueprintType, meta=(HasNativeMake="Engine.KismetMathLibrary.MakeVector2D", HasNativeBreak="Engine.KismetMathLibrary.BreakVector2D"))
 struct FVector2D
 {
 	UPROPERTY(EditAnywhere, BlueprintReadWrite, Category=Vector2D, SaveGame)
@@ -252,7 +252,7 @@ struct FPlane : public FVector
 
 // An orthogonal rotation in 3d space.
 
-USTRUCT(immutable, noexport, BlueprintType, meta=(HasNativeMakeBreak="true"))
+USTRUCT(immutable, noexport, BlueprintType, meta=(HasNativeMake="true", HasNativeBreak="true"))
 struct FRotator
 {
 	UPROPERTY(EditAnywhere, BlueprintReadWrite, Category=Rotator, SaveGame)
@@ -606,7 +606,7 @@ struct FInterpCurveLinearColor
 
 /** Transform definition. */
 
-USTRUCT(noexport, BlueprintType, meta=(HasNativeMakeBreak="true"))
+USTRUCT(noexport, BlueprintType, meta=(HasNativeMake="Engine.KismetMathLibrary.MakeTransform", HasNativeBreak="Engine.KismetMathLibrary.BreakTransform"))
 struct FTransform
 {
 	UPROPERTY(EditAnywhere, BlueprintReadWrite, Category=Transform, SaveGame)
diff --git a/Engine/Source/Runtime/Engine/Classes/EdGraph/EdGraphPin.h b/Engine/Source/Runtime/Engine/Classes/EdGraph/EdGraphPin.h
--- a/Engine/Source/Runtime/Engine/Classes/EdGraph/EdGraphPin.h
+++ b/Engine/Source/Runtime/Engine/Classes/EdGraph/EdGraphPin.h
@@ -194,6 +194,14 @@ class UEdGraphPin : public UObject
 	UPROPERTY()
 	TArray<class UEdGraphPin*> LinkedTo;
 
+	/** The pins created when a pin is split and hidden */ 
+	UPROPERTY()
+	TArray<class UEdGraphPin*> SubPins;
+
+	/** The pin that was split and generated this pin */
+	UPROPERTY()
+	UEdGraphPin* ParentPin;
+
 #if WITH_EDITORONLY_DATA
 	/** If true, this connector is currently hidden. */
 	UPROPERTY()
diff --git a/Engine/Source/Runtime/Engine/Classes/EdGraph/EdGraphSchema.h b/Engine/Source/Runtime/Engine/Classes/EdGraph/EdGraphSchema.h
--- a/Engine/Source/Runtime/Engine/Classes/EdGraph/EdGraphSchema.h
+++ b/Engine/Source/Runtime/Engine/Classes/EdGraph/EdGraphSchema.h
@@ -610,11 +610,17 @@ class ENGINE_API UEdGraphSchema : public UObject
 	 */
 	virtual void BreakSinglePinLink(UEdGraphPin* SourcePin, UEdGraphPin* TargetPin);
 
+	/** Split a pin in to subelements */
+	virtual void SplitPin(UEdGraphPin* Pin) const { };
+
+	/** Collapses a pin and its siblings back in to the original pin */
+	virtual void RecombinePin(UEdGraphPin* Pin) const { };
+
 	/** Break links on this pin and create links instead on MoveToPin */
-	virtual FPinConnectionResponse MovePinLinks(UEdGraphPin& MoveFromPin, UEdGraphPin& MoveToPin, bool bIsIntermeadiateMove = false) const;
+	virtual FPinConnectionResponse MovePinLinks(UEdGraphPin& MoveFromPin, UEdGraphPin& MoveToPin, bool bIsIntermediateMove = false) const;
 
 	/** Copies pin links from one pin to another without breaking the original links */
-	virtual FPinConnectionResponse CopyPinLinks(UEdGraphPin& CopyFromPin, UEdGraphPin& CopyToPin, bool bIsIntermeadiateCopy = false) const;
+	virtual FPinConnectionResponse CopyPinLinks(UEdGraphPin& CopyFromPin, UEdGraphPin& CopyToPin, bool bIsIntermediateCopy = false) const;
 
 	/** Is self pin type? */
 	virtual bool IsSelfPin(const UEdGraphPin& Pin) const   {return false;}
diff --git a/Engine/Source/Runtime/Engine/Classes/Engine/EngineTypes.h b/Engine/Source/Runtime/Engine/Classes/Engine/EngineTypes.h
--- a/Engine/Source/Runtime/Engine/Classes/Engine/EngineTypes.h
+++ b/Engine/Source/Runtime/Engine/Classes/Engine/EngineTypes.h
@@ -1271,7 +1271,7 @@ struct FPrimitiveMaterialRef
 };
 
 /** Structure containing information about one hit of the trace */
-USTRUCT(BlueprintType, meta=(HasNativeMakeBreak="true"))
+USTRUCT(BlueprintType, meta=(HasNativeBreak="Engine.GameplayStatics.BreakHitResult"))
 struct ENGINE_API FHitResult
 {
 	GENERATED_USTRUCT_BODY()
@@ -2108,7 +2108,7 @@ struct FRepAttachment
 
 	UPROPERTY()
 	FName AttachSocket;
-	
+
 	UPROPERTY()
 	class USceneComponent* AttachComponent;
 
@@ -2356,7 +2356,7 @@ namespace EComponentMobility
 		// - Totally Dynamic
 		// - Allows Dynamic Shadows
 		// - Slowest Rendering
-		Movable
+		Movable		
 	};
 }
 
diff --git a/Engine/Source/Runtime/Engine/Private/EdGraph/EdGraph.cpp b/Engine/Source/Runtime/Engine/Private/EdGraph/EdGraph.cpp
--- a/Engine/Source/Runtime/Engine/Private/EdGraph/EdGraph.cpp
+++ b/Engine/Source/Runtime/Engine/Private/EdGraph/EdGraph.cpp
@@ -174,9 +174,9 @@ void UEdGraph::NotifyGraphChanged(const FEdGraphEditAction& InAction)
 void UEdGraph::MoveNodesToAnotherGraph(UEdGraph* DestinationGraph, bool bIsLoading)
 {
 	// Move one node over at a time
-	while (Nodes.Num())
+	for (UEdGraphNode* Node : Nodes)
 	{
-		if (UEdGraphNode* Node = Nodes.Pop())
+		if (Node)
 		{
 			// Let the name be autogenerated, to automatically avoid naming conflicts
 			// Since this graph is always going to come from a cloned source graph, user readable names can come from the remap stored in a MessageLog.
@@ -187,6 +187,7 @@ void UEdGraph::MoveNodesToAnotherGraph(UEdGraph* DestinationGraph, bool bIsLoadi
 			DestinationGraph->Nodes.Add(Node);
 		}
 	}
+	Nodes.Empty();
 
 	DestinationGraph->NotifyGraphChanged();
 	NotifyGraphChanged();
diff --git a/Engine/Source/Runtime/Engine/Private/EdGraph/EdGraphPin.cpp b/Engine/Source/Runtime/Engine/Private/EdGraph/EdGraphPin.cpp
--- a/Engine/Source/Runtime/Engine/Private/EdGraph/EdGraphPin.cpp
+++ b/Engine/Source/Runtime/Engine/Private/EdGraph/EdGraphPin.cpp
@@ -169,6 +169,12 @@ void UEdGraphPin::CopyPersistentDataFromOldPin(const UEdGraphPin& SourcePin)
 		OtherPin->MakeLinkTo(this);
 	}
 
+	// If the source pin is split, then split the new one
+	if (SourcePin.SubPins.Num() > 0)
+	{
+		GetOwningNode()->GetSchema()->SplitPin(this);
+	}
+
 #if WITH_EDITORONLY_DATA
 	// Copy advanced visibility property, since it can be changed by user
 	bAdvancedView = SourcePin.bAdvancedView;
diff --git a/Engine/Source/Runtime/Engine/Private/EdGraph/EdGraphSchema.cpp b/Engine/Source/Runtime/Engine/Private/EdGraph/EdGraphSchema.cpp
--- a/Engine/Source/Runtime/Engine/Private/EdGraph/EdGraphSchema.cpp
+++ b/Engine/Source/Runtime/Engine/Private/EdGraph/EdGraphSchema.cpp
@@ -402,11 +402,11 @@ void UEdGraphSchema::BreakSinglePinLink(UEdGraphPin* SourcePin, UEdGraphPin* Tar
 #endif	//#if WITH_EDITOR
 }
 
-FPinConnectionResponse UEdGraphSchema::MovePinLinks(UEdGraphPin& MoveFromPin, UEdGraphPin& MoveToPin, bool bIsIntermeadiateMove) const
+FPinConnectionResponse UEdGraphSchema::MovePinLinks(UEdGraphPin& MoveFromPin, UEdGraphPin& MoveToPin, bool bIsIntermediateMove) const
 {
 #if WITH_EDITOR
-	ensureMsg(bIsIntermeadiateMove || !MoveToPin.GetOwningNode()->GetGraph()->HasAnyFlags(RF_Transient),
-		TEXT("When moving to an intermeadiate pin, use FKismetCompilerContext::MovePinLinksToIntermediate() instead of UEdGraphSchema::MovePinLinks()"));
+	ensureMsg(bIsIntermediateMove || !MoveToPin.GetOwningNode()->GetGraph()->HasAnyFlags(RF_Transient),
+		TEXT("When moving to an Intermediate pin, use FKismetCompilerContext::MovePinLinksToIntermediate() instead of UEdGraphSchema::MovePinLinks()"));
 #endif // #if WITH_EDITOR
 
 	FPinConnectionResponse FinalResponse = FPinConnectionResponse(CONNECT_RESPONSE_MAKE, TEXT(""));
@@ -435,11 +435,11 @@ FPinConnectionResponse UEdGraphSchema::MovePinLinks(UEdGraphPin& MoveFromPin, UE
 	return FinalResponse;
 }
 
-FPinConnectionResponse UEdGraphSchema::CopyPinLinks(UEdGraphPin& CopyFromPin, UEdGraphPin& CopyToPin, bool bIsIntermeadiateCopy) const
+FPinConnectionResponse UEdGraphSchema::CopyPinLinks(UEdGraphPin& CopyFromPin, UEdGraphPin& CopyToPin, bool bIsIntermediateCopy) const
 {
 #if WITH_EDITOR
-	ensureMsg(bIsIntermeadiateCopy || !CopyToPin.GetOwningNode()->GetGraph()->HasAnyFlags(RF_Transient),
-		TEXT("When copying to an intermeadiate pin, use FKismetCompilerContext::CopyPinLinksToIntermediate() instead of UEdGraphSchema::CopyPinLinks()"));
+	ensureMsg(bIsIntermediateCopy || !CopyToPin.GetOwningNode()->GetGraph()->HasAnyFlags(RF_Transient),
+		TEXT("When copying to an Intermediate pin, use FKismetCompilerContext::CopyPinLinksToIntermediate() instead of UEdGraphSchema::CopyPinLinks()"));
 #endif // #if WITH_EDITOR
 
 	FPinConnectionResponse FinalResponse = FPinConnectionResponse(CONNECT_RESPONSE_MAKE, TEXT(""));
