Fix issues with user defined structs and pin splitting

diff --git a/Engine/Source/Editor/KismetCompiler/Private/UserDefinedStructureCompilerUtils.cpp b/Engine/Source/Editor/KismetCompiler/Private/UserDefinedStructureCompilerUtils.cpp
--- a/Engine/Source/Editor/KismetCompiler/Private/UserDefinedStructureCompilerUtils.cpp
+++ b/Engine/Source/Editor/KismetCompiler/Private/UserDefinedStructureCompilerUtils.cpp
@@ -264,7 +264,9 @@ void FUserDefinedStructureCompilerUtils::CompileStruct(class UUserDefinedStruct*
 {
 	if (FStructureEditorUtils::UserDefinedStructEnabled() && Struct)
 	{
+		TSet<UBlueprint*> BlueprintsThatHaveBeenRecompiled;
 		TSet<UBlueprint*> BlueprintsToRecompile;
+
 		TArray<UUserDefinedStruct*> ChangedStructs; 
 
 		if (FUserDefinedStructureCompilerInner::ShouldBeCompiled(Struct) || bForceRecompile)
@@ -282,19 +284,55 @@ void FUserDefinedStructureCompilerUtils::CompileStruct(class UUserDefinedStruct*
 		FUserDefinedStructureCompilerInner::BuildDependencyMapAndCompile(ChangedStructs, MessageLog);
 
 		// UPDATE ALL THINGS DEPENDENT ON COMPILED STRUCTURES
-		for (TObjectIterator<UK2Node_StructOperation> It(RF_Transient | RF_PendingKill | RF_ClassDefaultObject, true); It && ChangedStructs.Num(); ++It)
+		for (TObjectIterator<UK2Node> It(RF_Transient | RF_PendingKill | RF_ClassDefaultObject, true); It && ChangedStructs.Num(); ++It)
 		{
-			UK2Node_StructOperation* Node = *It;
-			if (Node && !Node->HasAnyFlags(RF_Transient|RF_PendingKill))
+			bool bReconstruct = false;
+
+			UK2Node* Node = *It;
+
+			if (Node && !Node->HasAnyFlags(RF_Transient | RF_PendingKill))
+			{
+				// If this is a struct operation node operation on the changed struct we must reconstruct
+				if (UK2Node_StructOperation* StructOpNode = Cast<UK2Node_StructOperation>(Node))
+				{
+					UUserDefinedStruct* StructInNode = Cast<UUserDefinedStruct>(StructOpNode->StructType);
+					if (StructInNode && ChangedStructs.Contains(StructInNode))
+					{
+						bReconstruct = true;
+					}
+				}
+				if (!bReconstruct)
+				{
+					// Look through the nodes pins and if any of them are split and the type of the split pin is a user defined struct we need to reconstruct
+					for (UEdGraphPin* Pin : Node->Pins)
+					{
+						if (Pin->SubPins.Num() > 0)
+						{
+							UUserDefinedStruct* StructType = Cast<UUserDefinedStruct>(Pin->PinType.PinSubCategoryObject.Get());
+							if (StructType && ChangedStructs.Contains(StructType))
+							{
+								bReconstruct = true;
+								break;
+							}
+						}
+
+					}
+				}
+			}
+
+			if (bReconstruct)
 			{
-				UUserDefinedStruct* StructInNode = Cast<UUserDefinedStruct>(Node->StructType);
-				if (StructInNode && ChangedStructs.Contains(StructInNode))
+				if (UBlueprint* FoundBlueprint = Node->GetBlueprint())
 				{
-					if (UBlueprint* FoundBlueprint = Node->GetBlueprint())
+					// The blueprint skeleton needs to be updated before we reconstruct the node
+					// or else we may have member references that point to the old skeleton
+					if (!BlueprintsThatHaveBeenRecompiled.Contains(FoundBlueprint))
 					{
-						Node->ReconstructNode();
-						BlueprintsToRecompile.Add(FoundBlueprint);
+						BlueprintsThatHaveBeenRecompiled.Add(FoundBlueprint);
+						BlueprintsToRecompile.Remove(FoundBlueprint);
+						FBlueprintEditorUtils::MarkBlueprintAsStructurallyModified(FoundBlueprint);
 					}
+					Node->ReconstructNode();
 				}
 			}
 		}
