Kismet compiler will successfully pre-compile function graphs during a skeleton compile into UFunctions even if they are not successfully validated.

#ttp 338686 - LIVE: CRASH: BP: Adding a ForEachLoop to a blueprint function that contains Local variables

diff --git a/Engine/Source/Editor/KismetCompiler/Private/KismetCompiler.cpp b/Engine/Source/Editor/KismetCompiler/Private/KismetCompiler.cpp
--- a/Engine/Source/Editor/KismetCompiler/Private/KismetCompiler.cpp
+++ b/Engine/Source/Editor/KismetCompiler/Private/KismetCompiler.cpp
@@ -2803,7 +2803,8 @@ void FKismetCompilerContext::ProcessOneFunctionGraph(UEdGraph* SourceGraph)
 
 	// First do some cursory validation (pin types match, inputs to outputs, pins never point to their parent node, etc...)
 	// If this fails we don't proceed any further to avoid crashes or infinite loops
-	if (ValidateGraphIsWellFormed(FunctionGraph))
+	// When compiling only the skeleton class, we want the UFunction to be generated and processed so it contains all the local variables, this is unsafe to do during any other compilation mode
+	if (ValidateGraphIsWellFormed(FunctionGraph) || CompileOptions.CompileType == EKismetCompileType::SkeletonOnly)
 	{
 		FKismetFunctionContext& Context = *new (FunctionList) FKismetFunctionContext(MessageLog, Schema, NewClass, Blueprint);
 		Context.SourceGraph = FunctionGraph;
