Rearranged the logic to prevent local variables from being considered as const.
#TTP 339185 - EDITOR: BP: Local variables are designated as const when used in const functions or blueprints.
#Branch UE4
#Proj  KismetCompiler

#Change Changed the tests in FNodeHandlingFunctor::ResolveAndRegisterScopedTerm to prevent local variables ever being marked as const

diff --git a/Engine/Source/Editor/KismetCompiler/Private/KismetCompilerMisc.cpp b/Engine/Source/Editor/KismetCompiler/Private/KismetCompilerMisc.cpp
--- a/Engine/Source/Editor/KismetCompiler/Private/KismetCompilerMisc.cpp
+++ b/Engine/Source/Editor/KismetCompiler/Private/KismetCompilerMisc.cpp
@@ -716,17 +716,16 @@ void FNodeHandlingFunctor::ResolveAndRegisterScopedTerm(FKismetFunctionContext&
 		Term->AssociatedVarProperty = BoundProperty;
 		Context.NetMap.Add(Net, Term);
 
-		// Read-only variables and variables in const classes are both const
-		if (BoundProperty->HasAnyPropertyFlags(CPF_BlueprintReadOnly) || Context.IsConstFunction())
-		{
-			Term->bIsConst = true;
-		}
-
 		// Check if the property is a local variable and mark it so
 		if( SearchScope == Context.Function && BoundProperty->GetOuter() == Context.Function)
 		{
 			Term->bIsLocal = true;
 		}
+		else if (BoundProperty->HasAnyPropertyFlags(CPF_BlueprintReadOnly) || Context.IsConstFunction())
+		{
+			// Read-only variables and variables in const classes are both const
+			Term->bIsConst = true;
+		}
 
 		// Resolve the context term
 		if (SelfPin != NULL)
