BlueprintFunctionLibrary cannot have (non-local) variables. ttp338346

diff --git a/Engine/Source/Editor/BlueprintGraph/Private/EdGraphSchema_K2.cpp b/Engine/Source/Editor/BlueprintGraph/Private/EdGraphSchema_K2.cpp
--- a/Engine/Source/Editor/BlueprintGraph/Private/EdGraphSchema_K2.cpp
+++ b/Engine/Source/Editor/BlueprintGraph/Private/EdGraphSchema_K2.cpp
@@ -3248,7 +3248,7 @@ bool UEdGraphSchema_K2::CanPromotePinToVariable( const UEdGraphPin& Pin ) const
 	const UK2Node* Node = Cast<UK2Node>(Pin.GetOwningNode());
 	const UBlueprint* OwningBlueprint = Node->GetBlueprint();
 	
-	if (OwningBlueprint && (OwningBlueprint->BlueprintType == BPTYPE_MacroLibrary))
+	if (!OwningBlueprint || (OwningBlueprint->BlueprintType == BPTYPE_MacroLibrary) || (OwningBlueprint->BlueprintType == BPTYPE_FunctionLibrary))
 	{
 		// Never allow promotion in macros, because there's not a scope to define them in
 		bCanPromote = false;
diff --git a/Engine/Source/Editor/KismetCompiler/Private/KismetCompiler.cpp b/Engine/Source/Editor/KismetCompiler/Private/KismetCompiler.cpp
--- a/Engine/Source/Editor/KismetCompiler/Private/KismetCompiler.cpp
+++ b/Engine/Source/Editor/KismetCompiler/Private/KismetCompiler.cpp
@@ -369,6 +369,12 @@ void FKismetCompilerContext::ValidateNode(const UEdGraphNode* Node) const
 /** Creates a class variable */
 UProperty* FKismetCompilerContext::CreateVariable(const FName VarName, const FEdGraphPinType& VarType)
 {
+	if (BPTYPE_FunctionLibrary == Blueprint->BlueprintType)
+	{
+		MessageLog.Error(*FString::Printf(*LOCTEXT("VariableInFunctionLibrary_Error", "The variable %s cannot be declared in FunctionLibrary @@").ToString(),
+			*VarName.ToString()), Blueprint);
+	}
+
 	UProperty* NewProperty = FKismetCompilerUtilities::CreatePropertyOnScope(NewClass, VarName, VarType, NewClass, 0, Schema, MessageLog);
 	if (NewProperty != NULL)
 	{
diff --git a/Engine/Source/Editor/UnrealEd/Private/Kismet2/BlueprintEditorUtils.cpp b/Engine/Source/Editor/UnrealEd/Private/Kismet2/BlueprintEditorUtils.cpp
--- a/Engine/Source/Editor/UnrealEd/Private/Kismet2/BlueprintEditorUtils.cpp
+++ b/Engine/Source/Editor/UnrealEd/Private/Kismet2/BlueprintEditorUtils.cpp
@@ -3237,7 +3237,7 @@ void FBlueprintEditorUtils::RenameMemberVariable(UBlueprint* Blueprint, const FN
 				}
 			}
 		}
-		else
+		else if (Blueprint && Blueprint->SimpleConstructionScript)
 		{
 			// Wasn't in the introduced variable list; try to find the associated SCS node
 			//@TODO: The SCS-generated variables should be in the variable list and have a link back;
