UDS sends notification about change while compilation.
ttp#339577 CRASH: EDITOR: UDS: Crash when operating on variables default values

diff --git a/Engine/Source/Editor/KismetCompiler/Private/UserDefinedStructureCompilerUtils.cpp b/Engine/Source/Editor/KismetCompiler/Private/UserDefinedStructureCompilerUtils.cpp
--- a/Engine/Source/Editor/KismetCompiler/Private/UserDefinedStructureCompilerUtils.cpp
+++ b/Engine/Source/Editor/KismetCompiler/Private/UserDefinedStructureCompilerUtils.cpp
@@ -199,7 +199,7 @@ struct FUserDefinedStructureCompilerInner
 		return true;
 	}
 
-	static void BuildDependencyMapAndCompile(TArray<UUserDefinedStruct*>& ChangedStructs, FCompilerResultsLog& MessageLog)
+	static void BuildDependencyMapAndCompile(const TArray<UUserDefinedStruct*>& ChangedStructs, FCompilerResultsLog& MessageLog)
 	{
 		struct FDependencyMapEntry
 		{
@@ -208,7 +208,7 @@ struct FUserDefinedStructureCompilerInner
 
 			FDependencyMapEntry() : Struct(NULL) {}
 
-			void Initialize(UUserDefinedStruct* ChangedStruct, TArray<UUserDefinedStruct*>& AllChangedStructs) 
+			void Initialize(UUserDefinedStruct* ChangedStruct, const TArray<UUserDefinedStruct*>& AllChangedStructs) 
 			{ 
 				Struct = ChangedStruct;
 				check(Struct);
@@ -226,7 +226,7 @@ struct FUserDefinedStructureCompilerInner
 		};
 
 		TArray<FDependencyMapEntry> DependencyMap;
-		for (auto Iter = ChangedStructs.CreateIterator(); Iter; ++Iter)
+		for (auto Iter = ChangedStructs.CreateConstIterator(); Iter; ++Iter)
 		{
 			DependencyMap.Add(FDependencyMapEntry());
 			DependencyMap.Last().Initialize(*Iter, ChangedStructs);
@@ -276,8 +276,13 @@ void FUserDefinedStructureCompilerUtils::CompileStruct(class UUserDefinedStruct*
 
 		for (int32 StructIdx = 0; StructIdx < ChangedStructs.Num(); ++StructIdx)
 		{
-			FUserDefinedStructureCompilerInner::ReplaceStructWithTempDuplicate(ChangedStructs[StructIdx], BlueprintsToRecompile, ChangedStructs);
-			ChangedStructs[StructIdx]->Status = EUserDefinedStructureStatus::UDSS_Dirty;
+			UUserDefinedStruct* ChangedStruct = ChangedStructs[StructIdx];
+			if (ChangedStruct)
+			{
+				FStructureEditorUtils::FStructEditorManager::Get().PreChange(ChangedStruct);
+				FUserDefinedStructureCompilerInner::ReplaceStructWithTempDuplicate(ChangedStruct, BlueprintsToRecompile, ChangedStructs);
+				ChangedStruct->Status = EUserDefinedStructureStatus::UDSS_Dirty;
+			}
 		}
 
 		// COMPILE IN PROPER ORDER
@@ -341,6 +346,16 @@ void FUserDefinedStructureCompilerUtils::CompileStruct(class UUserDefinedStruct*
 		{
 			FBlueprintEditorUtils::MarkBlueprintAsStructurallyModified(*BPIter);
 		}
+
+		// UPDATE UI
+		for (auto ChangedStruct : ChangedStructs)
+		{
+			if (ChangedStruct)
+			{
+				FStructureEditorUtils::FStructEditorManager::Get().PostChange(ChangedStruct);
+				ChangedStruct->MarkPackageDirty();
+			}
+		}
 	}
 }
 
diff --git a/Engine/Source/Editor/UnrealEd/Private/Kismet2/StructureEditorUtils.cpp b/Engine/Source/Editor/UnrealEd/Private/Kismet2/StructureEditorUtils.cpp
--- a/Engine/Source/Editor/UnrealEd/Private/Kismet2/StructureEditorUtils.cpp
+++ b/Engine/Source/Editor/UnrealEd/Private/Kismet2/StructureEditorUtils.cpp
@@ -460,10 +460,8 @@ void FStructureEditorUtils::OnStructureChanged(UUserDefinedStruct* Struct)
 {
 	if (Struct)
 	{
-		FStructEditorManager::Get().PreChange(Struct);
 		Struct->Status = EUserDefinedStructureStatus::UDSS_Dirty;
 		CompileStructure(Struct);
-		FStructEditorManager::Get().PostChange(Struct);
 		Struct->MarkPackageDirty();
 	}
 }
