VM - fix for passing UObject* in structures

diff --git a/Engine/Source/Editor/KismetCompiler/Private/KismetCompilerVMBackend.cpp b/Engine/Source/Editor/KismetCompiler/Private/KismetCompilerVMBackend.cpp
--- a/Engine/Source/Editor/KismetCompiler/Private/KismetCompilerVMBackend.cpp
+++ b/Engine/Source/Editor/KismetCompiler/Private/KismetCompilerVMBackend.cpp
@@ -481,6 +481,10 @@ public:
 						{
 							NewTerm.TextLiteral = FText::FromString(NewTerm.Name);
 						}
+						else if(Prop->IsA(UObjectProperty::StaticClass()))
+						{
+							NewTerm.ObjectLiteral = Cast<UObjectProperty>(Prop)->GetObjectPropertyValue(Prop->ContainerPtrToValuePtr<void>(StructData));
+						}
 
 						EmitTermExpr(&NewTerm, Prop);
 					}
