UMG - Continued improvement of UMG and the blueprint layer.  Now UMG Blueprints support multicast delegate binding in the blueprint.  Fixed a nasty serialization issue with it not correctly serializing the widget tree.
Blueprint - The underlying framework now allows derrived blueprints to say if they support event binding, instead of it being hardcoded to only work for AActors.
Blueprint - Function calls in the context menu on selected "Components" now works on UObject based blueprints.  There was no reason to restrict tthat to AActors.

diff --git a/Engine/Config/BaseEngine.ini b/Engine/Config/BaseEngine.ini
--- a/Engine/Config/BaseEngine.ini
+++ b/Engine/Config/BaseEngine.ini
@@ -562,6 +562,9 @@ bUseStreamingPause=false
 +ActiveClassRedirects=(OldClassName="SpriteComponent",NewClassName="BillboardComponent")
 +ActiveClassRedirects=(OldClassName="MaterialSpriteComponent",NewClassName="MaterialBillboardComponent")
 
+; Widget Blueprint Redirects
++ActiveClassRedirects=(OldClassName="WidgetBlueprint",NewClassName="/Script/UMGEditor.WidgetBlueprint")
+
 ;EAxis rename
 +EnumRedirects=(EnumName="EAxis",OldEnumEntry="AXIS_X",NewEnumEntry="EAxis::X")
 +EnumRedirects=(EnumName="EAxis",OldEnumEntry="AXIS_Y",NewEnumEntry="EAxis::Y")
diff --git a/Engine/Source/Editor/BlueprintGraph/Private/K2ActionMenuBuilder.cpp b/Engine/Source/Editor/BlueprintGraph/Private/K2ActionMenuBuilder.cpp
--- a/Engine/Source/Editor/BlueprintGraph/Private/K2ActionMenuBuilder.cpp
+++ b/Engine/Source/Editor/BlueprintGraph/Private/K2ActionMenuBuilder.cpp
@@ -137,7 +137,7 @@ static void AddSpawnActorNodeAction(FGraphActionListBuilderBase& ContextMenuBuil
 static void AddGetDataTableRowNodeAction(FGraphActionListBuilderBase& ContextMenuBuilder, const FString& FunctionCategory)
 {
 	FString const GetDataTableRowCategory = FunctionCategory + K2ActionCategories::SubCategoryDelim + K2ActionCategories::UtilitiesCategory;
-    
+	
 	// Add the new GetDataTableRow node
 	{
 		UK2Node* NodeTemplate = ContextMenuBuilder.CreateTemplateNode<UK2Node_GetDataTableRow>();
@@ -835,8 +835,8 @@ void FK2ActionMenuBuilder::GetContextAllowedNodeTypes(FBlueprintGraphActionListB
 				GetInputNodes(ContextMenuBuilder, true);
 			}
 		}
-        
-        AddGetDataTableRowNodeAction(ContextMenuBuilder, K2ActionCategories::GenericFunctionCategory);
+		
+		AddGetDataTableRowNodeAction(ContextMenuBuilder, K2ActionCategories::GenericFunctionCategory);
 
 		// Add struct make/break nodes
 		GetStructActions( ContextMenuBuilder );
@@ -904,13 +904,17 @@ void FK2ActionMenuBuilder::GetContextAllowedNodeTypes(FBlueprintGraphActionListB
 			GetFunctionCallsOnSelectedActors(ContextMenuBuilder);
 		}
 		// Non level script
-		else if (FBlueprintEditorUtils::IsActorBased(Blueprint))
+		else
 		{
-			GetAddComponentActionsUsingSelectedAssets(ContextMenuBuilder);
-			GetAddComponentClasses(ContextMenuBuilder.Blueprint, ContextMenuBuilder);
+			if ( FBlueprintEditorUtils::IsActorBased(Blueprint) )
+			{
+				GetAddComponentActionsUsingSelectedAssets(ContextMenuBuilder);
+				GetAddComponentClasses(ContextMenuBuilder.Blueprint, ContextMenuBuilder);
+			}
+
 			GetFunctionCallsOnSelectedComponents(ContextMenuBuilder);
 
-			if(bAllowEvents)
+			if ( bAllowEvents && Blueprint->AllowsDynamicBinding() )
 			{
 				GetBoundEventsFromComponentSelection(ContextMenuBuilder);			
 			}
@@ -1138,8 +1142,8 @@ void FK2ActionMenuBuilder::GetPinAllowedNodeTypes(FBlueprintGraphActionListBuild
 					GetInputNodes(ContextMenuBuilder, bIncludeEvents);
 				}
 			}
-            
-            AddGetDataTableRowNodeAction(ContextMenuBuilder, K2ActionCategories::GenericFunctionCategory);
+			
+			AddGetDataTableRowNodeAction(ContextMenuBuilder, K2ActionCategories::GenericFunctionCategory);
 
 			// for output pins, add macro flow control as a connection option
 			if ( FromPin.Direction == EGPD_Output )
@@ -2348,8 +2352,7 @@ void FK2ActionMenuBuilder::GetBoundEventsFromComponentSelection(FBlueprintGraphA
 	{
 		UObjectProperty* CompProperty = Cast<UObjectProperty>(ContextMenuBuilder.SelectedObjects[SelectIdx]);
 		if(	CompProperty != NULL && 
-			CompProperty->PropertyClass != NULL &&
-			CompProperty->PropertyClass->IsChildOf(UActorComponent::StaticClass()) )
+			CompProperty->PropertyClass != NULL )
 		{
 			const FString PropertyName = GetDefault<UEditorStyleSettings>()->bShowFriendlyNames ? UEditorEngine::GetFriendlyName(CompProperty) : CompProperty->GetName();
 			const FString CurrentCategory = K2ActionCategories::AddEventForPrefix + PropertyName;
@@ -2369,8 +2372,7 @@ void FK2ActionMenuBuilder::GetFunctionCallsOnSelectedComponents(FBlueprintGraphA
 	{
 		UObjectProperty* CompProperty = Cast<UObjectProperty>(ContextMenuBuilder.SelectedObjects[SelectIdx]);
 		if(	CompProperty != NULL && 
-			CompProperty->PropertyClass != NULL &&
-			CompProperty->PropertyClass->IsChildOf(UActorComponent::StaticClass()) )
+			CompProperty->PropertyClass != NULL )
 		{
 			SelectedCompProperty = CompProperty;
 			break;
diff --git a/Engine/Source/Editor/KismetCompiler/Private/KismetCompiler.cpp b/Engine/Source/Editor/KismetCompiler/Private/KismetCompiler.cpp
--- a/Engine/Source/Editor/KismetCompiler/Private/KismetCompiler.cpp
+++ b/Engine/Source/Editor/KismetCompiler/Private/KismetCompiler.cpp
@@ -115,40 +115,24 @@ void FKismetCompilerContext::SpawnNewClass(const FString& NewClassName)
 	}
 }
 
-struct FSubobjectCollection
-{ 
-private:
-	TSet<const UObject*> Collection;
-
-public:
-	void AddObject(const UObject* const InObject)
-	{
-		if (InObject)
-		{
-			Collection.Add(InObject);
-			TArray<UObject*> Subobjects;
-			GetObjectsWithOuter(InObject, Subobjects, true);
-			for (auto SubObject : Subobjects)
-			{
-				Collection.Add(SubObject);
-			}
-		}
-	}
-
-	template<typename TOBJ>
-	void AddObjects(const TArray<TOBJ*>& InObjects)
+void FKismetCompilerContext::FSubobjectCollection::AddObject(const UObject* const InObject)
+{
+	if ( InObject )
 	{
-		for (const auto ObjPtr : InObjects)
+		Collection.Add(InObject);
+		TArray<UObject*> Subobjects;
+		GetObjectsWithOuter(InObject, Subobjects, true);
+		for ( auto SubObject : Subobjects )
 		{
-			AddObject(ObjPtr);
+			Collection.Add(SubObject);
 		}
 	}
+}
 
-	bool operator()(const UObject* const RemovalCandidate) const
-	{
-		return (NULL != Collection.Find(RemovalCandidate));
-	}
-};
+bool FKismetCompilerContext::FSubobjectCollection::operator()(const UObject* const RemovalCandidate) const
+{
+	return ( NULL != Collection.Find(RemovalCandidate) );
+}
 
 void FKismetCompilerContext::CleanAndSanitizeClass(UBlueprintGeneratedClass* ClassToClean, UObject*& OldCDO)
 {
@@ -163,7 +147,7 @@ void FKismetCompilerContext::CleanAndSanitizeClass(UBlueprintGeneratedClass* Cla
 		ParentClass = UObject::StaticClass();
 	}
 	TransientClass->ClassAddReferencedObjects = ParentClass->AddReferencedObjects;
-    
+	
 	NewClass = ClassToClean;
 	OldCDO = ClassToClean->ClassDefaultObject; // we don't need to create the CDO at this point
 	
@@ -179,45 +163,10 @@ void FKismetCompilerContext::CleanAndSanitizeClass(UBlueprintGeneratedClass* Cla
 	TArray<UObject*> ClassSubObjects;
 	GetObjectsWithOuter(ClassToClean, ClassSubObjects, true);
 
-	{	// Save subobjects, that won't be regenerated.
-		FSubobjectCollection SubObjectsToSave;
-		SubObjectsToSave.AddObjects(Blueprint->ComponentTemplates);
-		SubObjectsToSave.AddObjects(Blueprint->Timelines);
-		if (Blueprint->SimpleConstructionScript)
-		{
-			SubObjectsToSave.AddObject(Blueprint->SimpleConstructionScript);
-			if (const USCS_Node* DefaultScene = Blueprint->SimpleConstructionScript->GetDefaultSceneRootNode())
-	{
-				SubObjectsToSave.AddObject(DefaultScene->ComponentTemplate);
-	}
-
-			TArray<USCS_Node*> SCSNodes = Blueprint->SimpleConstructionScript->GetAllNodes();
-			for (auto SCSNode : SCSNodes)
-	{
-				SubObjectsToSave.AddObject(SCSNode->ComponentTemplate);
-			}
-	}
 	{
-			TSet<class UCurveBase*> Curves;
-			for (auto Timeline : Blueprint->Timelines)
-			{
-				if (Timeline)
-		{
-					Timeline->GetAllCurves(Curves);
-		}
-			}
-			for (auto Component : Blueprint->ComponentTemplates)
-			{
-				if (auto TimelineComponent = Cast<UTimelineComponent>(Component))
-		{
-					TimelineComponent->GetAllCurves(Curves);
-		}
-			}
-			for (auto Curve : Curves)
-			{
-				SubObjectsToSave.AddObject(Curve);
-			}
-		}
+		// Save subobjects, that won't be regenerated.
+		FSubobjectCollection SubObjectsToSave;
+		SaveSubObjectsFromCleanAndSanitizeClass(SubObjectsToSave, ClassToClean, OldCDO);
 
 		ClassSubObjects.RemoveAllSwap(SubObjectsToSave);
 	}
@@ -246,6 +195,49 @@ void FKismetCompilerContext::CleanAndSanitizeClass(UBlueprintGeneratedClass* Cla
 	ClassToClean->DebugData = FBlueprintDebugData();
 }
 
+void FKismetCompilerContext::SaveSubObjectsFromCleanAndSanitizeClass(FSubobjectCollection& SubObjectsToSave, UBlueprintGeneratedClass* ClassToClean, UObject*& OldCDO)
+{
+	SubObjectsToSave.AddObjects(Blueprint->ComponentTemplates);
+	SubObjectsToSave.AddObjects(Blueprint->Timelines);
+
+	if ( Blueprint->SimpleConstructionScript )
+	{
+		SubObjectsToSave.AddObject(Blueprint->SimpleConstructionScript);
+		if ( const USCS_Node* DefaultScene = Blueprint->SimpleConstructionScript->GetDefaultSceneRootNode() )
+		{
+			SubObjectsToSave.AddObject(DefaultScene->ComponentTemplate);
+		}
+
+		TArray<USCS_Node*> SCSNodes = Blueprint->SimpleConstructionScript->GetAllNodes();
+		for ( auto SCSNode : SCSNodes )
+		{
+			SubObjectsToSave.AddObject(SCSNode->ComponentTemplate);
+		}
+	}
+
+	{
+		TSet<class UCurveBase*> Curves;
+		for ( auto Timeline : Blueprint->Timelines )
+		{
+			if ( Timeline )
+			{
+				Timeline->GetAllCurves(Curves);
+			}
+		}
+		for ( auto Component : Blueprint->ComponentTemplates )
+		{
+			if ( auto TimelineComponent = Cast<UTimelineComponent>(Component) )
+			{
+				TimelineComponent->GetAllCurves(Curves);
+			}
+		}
+		for ( auto Curve : Curves )
+		{
+			SubObjectsToSave.AddObject(Curve);
+		}
+	}
+}
+
 void FKismetCompilerContext::PostCreateSchema()
 {
 	NodeHandlers.Add(UEdGraphNode_Comment::StaticClass(), new FNodeHandlingFunctor(*this));
@@ -2202,20 +2194,20 @@ void FKismetCompilerContext::CreateFunctionStubForEvent(UK2Node_Event* SrcEventN
 
 	// If this is a customizable event, make sure to copy over the user defined pins
 	if (UK2Node_CustomEvent const* SrcCustomEventNode = Cast<UK2Node_CustomEvent const>(SrcEventNode))
- 	{
+	{
 		EntryNode->UserDefinedPins = SrcCustomEventNode->UserDefinedPins;
 		// CustomEvents may inherit net flags (so let's use their GetNetFlags() incase this is an override)
 		StubContext.MarkAsNetFunction(SrcCustomEventNode->GetNetFlags());
- 	}
+	}
 	EntryNode->AllocateDefaultPins();
 
 	// Confirm that the event node matches the latest function signature, which the newly created EntryNode should have
- 	if( !SrcEventNode->IsFunctionEntryCompatible(EntryNode) )
- 	{
- 		// There is no match, so the function parameters must have changed.  Throw an error, and force them to refresh
- 		MessageLog.Error(*FString::Printf(*LOCTEXT("EventNodeOutOfDate_Error", "Event node @@ is out-of-date.  Please refresh it.").ToString()), SrcEventNode);
- 		return;
- 	}
+	if( !SrcEventNode->IsFunctionEntryCompatible(EntryNode) )
+	{
+		// There is no match, so the function parameters must have changed.  Throw an error, and force them to refresh
+		MessageLog.Error(*FString::Printf(*LOCTEXT("EventNodeOutOfDate_Error", "Event node @@ is out-of-date.  Please refresh it.").ToString()), SrcEventNode);
+		return;
+	}
 
 	// Copy each event parameter to the assignment node, if there are any inputs
 	UK2Node_VariableSet* AssignmentNode = NULL;
@@ -2260,7 +2252,7 @@ void FKismetCompilerContext::CreateFunctionStubForEvent(UK2Node_Event* SrcEventN
 	CallIntoUbergraph->NodePosX = 300;
 
 	// Use the ExecuteUbergraph base function to generate the pins...
- 	CallIntoUbergraph->FunctionReference.SetExternalMember(Schema->FN_ExecuteUbergraphBase, UObject::StaticClass());
+	CallIntoUbergraph->FunctionReference.SetExternalMember(Schema->FN_ExecuteUbergraphBase, UObject::StaticClass());
 	CallIntoUbergraph->AllocateDefaultPins();
 	
 	// ...then swap to the generated version for this level
@@ -2953,25 +2945,25 @@ void FKismetCompilerContext::Compile()
 	UBlueprintGeneratedClass* TargetClass = Cast<UBlueprintGeneratedClass>(TargetUClass);
 
 	// >>> Backwards Compatibility: Make sure that skeleton generated classes have the proper "SKEL_" naming convention
- 	const FString SkeletonPrefix(TEXT("SKEL_"));
- 	if( bIsSkeletonOnly && TargetClass && !TargetClass->GetName().StartsWith(SkeletonPrefix) )
- 	{
- 		FString NewName = SkeletonPrefix + TargetClass->GetName();
+	const FString SkeletonPrefix(TEXT("SKEL_"));
+	if( bIsSkeletonOnly && TargetClass && !TargetClass->GetName().StartsWith(SkeletonPrefix) )
+	{
+		FString NewName = SkeletonPrefix + TargetClass->GetName();
  
- 		// Ensure we have a free name for this class
- 		UClass* AnyClassWithGoodName = (UClass*)StaticFindObject(UClass::StaticClass(), Blueprint->GetOutermost(), *NewName, false);
- 		if( AnyClassWithGoodName )
- 		{
- 			// Special Case:  If the CDO of the class has become dissociated from its actual CDO, attempt to find the proper named CDO, and get rid of it.
- 			if( AnyClassWithGoodName->ClassDefaultObject == TargetClass->ClassDefaultObject )
- 			{
- 				AnyClassWithGoodName->ClassDefaultObject = NULL;
- 				FString DefaultObjectName = FString(DEFAULT_OBJECT_PREFIX) + NewName;
- 				AnyClassWithGoodName->ClassDefaultObject = (UObject*)StaticFindObject(UObject::StaticClass(), Blueprint->GetOutermost(), *DefaultObjectName, false);
- 			}
+		// Ensure we have a free name for this class
+		UClass* AnyClassWithGoodName = (UClass*)StaticFindObject(UClass::StaticClass(), Blueprint->GetOutermost(), *NewName, false);
+		if( AnyClassWithGoodName )
+		{
+			// Special Case:  If the CDO of the class has become dissociated from its actual CDO, attempt to find the proper named CDO, and get rid of it.
+			if( AnyClassWithGoodName->ClassDefaultObject == TargetClass->ClassDefaultObject )
+			{
+				AnyClassWithGoodName->ClassDefaultObject = NULL;
+				FString DefaultObjectName = FString(DEFAULT_OBJECT_PREFIX) + NewName;
+				AnyClassWithGoodName->ClassDefaultObject = (UObject*)StaticFindObject(UObject::StaticClass(), Blueprint->GetOutermost(), *DefaultObjectName, false);
+			}
  
 			// Get rid of the old class to make room for renaming our class to the final SKEL name
- 			FKismetCompilerUtilities::ConsignToOblivion(AnyClassWithGoodName, Blueprint->bIsRegeneratingOnLoad);
+			FKismetCompilerUtilities::ConsignToOblivion(AnyClassWithGoodName, Blueprint->bIsRegeneratingOnLoad);
 
 			// Update the refs to the old SKC
 			TMap<UObject*, UObject*> ClassReplacementMap;
@@ -2982,11 +2974,11 @@ void FKismetCompilerContext::Compile()
 			{
 				FArchiveReplaceObjectRef<UObject> ReplaceInBlueprintAr(AllGraphs[i], ClassReplacementMap, /*bNullPrivateRefs=*/ false, /*bIgnoreOuterRef=*/ false, /*bIgnoreArchetypeRef=*/ false);
 			}
- 		}
+		}
  
- 		ERenameFlags RenameFlags = (REN_DontCreateRedirectors|REN_NonTransactional|(Blueprint->bIsRegeneratingOnLoad ? REN_ForceNoResetLoaders : 0));
- 		TargetClass->Rename(*NewName, NULL, RenameFlags);
- 	}
+		ERenameFlags RenameFlags = (REN_DontCreateRedirectors|REN_NonTransactional|(Blueprint->bIsRegeneratingOnLoad ? REN_ForceNoResetLoaders : 0));
+		TargetClass->Rename(*NewName, NULL, RenameFlags);
+	}
 	// <<< End Backwards Compatibility
 
 	// >>> Backwards compatibility:  If SkeletonGeneratedClass == GeneratedClass, we need to make a new generated class the first time we need it
@@ -3380,29 +3372,29 @@ void FKismetCompilerContext::Compile()
 	}
 
 	// For full compiles, find other blueprints that may need refreshing, and mark them dirty, in case they try to run
-  	if( bIsFullCompile && !Blueprint->bIsRegeneratingOnLoad )
-  	{
-  		TArray<UObject*> AllBlueprints;
-  		GetObjectsOfClass(UBlueprint::StaticClass(), AllBlueprints, true);
+	if( bIsFullCompile && !Blueprint->bIsRegeneratingOnLoad )
+	{
+		TArray<UObject*> AllBlueprints;
+		GetObjectsOfClass(UBlueprint::StaticClass(), AllBlueprints, true);
   
-  		// Mark any blueprints that implement this interface as dirty
-  		for( auto CurrentObj = AllBlueprints.CreateIterator(); CurrentObj; ++CurrentObj )
-  		{
-  			UBlueprint* CurrentBP = Cast<UBlueprint>(*CurrentObj);
+		// Mark any blueprints that implement this interface as dirty
+		for( auto CurrentObj = AllBlueprints.CreateIterator(); CurrentObj; ++CurrentObj )
+		{
+			UBlueprint* CurrentBP = Cast<UBlueprint>(*CurrentObj);
   
-  			if( FBlueprintEditorUtils::IsBlueprintDependentOn(CurrentBP, Blueprint) )
-  			{
+			if( FBlueprintEditorUtils::IsBlueprintDependentOn(CurrentBP, Blueprint) )
+			{
 				CurrentBP->Status = BS_Dirty;
-  				FBlueprintEditorUtils::RefreshExternalBlueprintDependencyNodes(CurrentBP);
+				FBlueprintEditorUtils::RefreshExternalBlueprintDependencyNodes(CurrentBP);
 				CurrentBP->BroadcastChanged();
-  			}
-  		}
-  	}
+			}
+		}
+	}
 
 	// Clear out pseudo-local members that are only valid within a Compile call
 	UbergraphContext = NULL;
- 	CallsIntoUbergraph.Empty();
- 	TimelineToMemberVariableMap.Empty();
+	CallsIntoUbergraph.Empty();
+	TimelineToMemberVariableMap.Empty();
 
 
 	check(NewClass->PropertiesSize >= UObject::StaticClass()->PropertiesSize);
diff --git a/Engine/Source/Editor/KismetCompiler/Public/KismetCompiler.h b/Engine/Source/Editor/KismetCompiler/Public/KismetCompiler.h
--- a/Engine/Source/Editor/KismetCompiler/Public/KismetCompiler.h
+++ b/Engine/Source/Editor/KismetCompiler/Public/KismetCompiler.h
@@ -155,6 +155,31 @@ protected:
 	 */
 	virtual void CleanAndSanitizeClass(UBlueprintGeneratedClass* ClassToClean, UObject*& OldCDO);
 
+	struct KISMETCOMPILER_API FSubobjectCollection
+	{
+	private:
+		TSet<const UObject*> Collection;
+
+	public:
+		void AddObject(const UObject* const InObject);
+
+		template<typename TOBJ>
+		void AddObjects(const TArray<TOBJ*>& InObjects)
+		{
+			for ( const auto ObjPtr : InObjects )
+			{
+				AddObject(ObjPtr);
+			}
+		}
+
+		bool operator()(const UObject* const RemovalCandidate) const;
+	};
+
+	/**
+	 * Saves any SubObjects on the blueprint that need to survive the clean 
+	 */
+	virtual void SaveSubObjectsFromCleanAndSanitizeClass(FSubobjectCollection& SubObjectsToSave, UBlueprintGeneratedClass* ClassToClean, UObject*& OldCDO);
+
 	/** 
 	 * Checks a connection response, and errors if it didn't succeed (not public, 
 	 * users should be using MovePinLinksToIntermediate/CopyPinLinksToIntermediate 
diff --git a/Engine/Source/Editor/UMGEditor/Private/WidgetBlueprint.cpp b/Engine/Source/Editor/UMGEditor/Private/WidgetBlueprint.cpp
--- a/Engine/Source/Editor/UMGEditor/Private/WidgetBlueprint.cpp
+++ b/Engine/Source/Editor/UMGEditor/Private/WidgetBlueprint.cpp
@@ -51,6 +51,12 @@ UClass* UWidgetBlueprint::GetBlueprintClass() const
 }
 
 
+bool UWidgetBlueprint::AllowsDynamicBinding() const
+{
+	return true;
+}
+
+
 bool UWidgetBlueprint::ValidateGeneratedClass(const UClass* InClass)
 {
 	bool Result = Super::ValidateGeneratedClass(InClass);
@@ -66,15 +72,39 @@ bool UWidgetBlueprint::ValidateGeneratedClass(const UClass* InClass)
 		return false;
 	}
 
-	//if ( !ensure(Blueprint->WidgetTree && ( Blueprint->WidgetTree->GetOuter() == GeneratedClass )) )
-	//{
-	//	return false;
-	//}
+	if ( !ensure(Blueprint->WidgetTree && ( Blueprint->WidgetTree->GetOuter() == Blueprint )) )
+	{
+		return false;
+	}
+	else
+	{
+		TArray < UWidget* > AllWidgets;
+		Blueprint->WidgetTree->GetAllWidgets(AllWidgets);
+		for ( UWidget* Widget : AllWidgets )
+		{
+			if ( !ensure(Widget->GetOuter() == Blueprint->WidgetTree) )
+			{
+				return false;
+			}
+		}
+	}
 
-	//if ( !ensure(GeneratedClass->WidgetTree && ( GeneratedClass->WidgetTree->GetOuter() == GeneratedClass )) )
-	//{
-	//	return false;
-	//}
+	if ( !ensure(GeneratedClass->WidgetTree && ( GeneratedClass->WidgetTree->GetOuter() == GeneratedClass )) )
+	{
+		return false;
+	}
+	else
+	{
+		TArray < UWidget* > AllWidgets;
+		GeneratedClass->WidgetTree->GetAllWidgets(AllWidgets);
+		for ( UWidget* Widget : AllWidgets )
+		{
+			if ( !ensure(Widget->GetOuter() == GeneratedClass->WidgetTree) )
+			{
+				return false;
+			}
+		}
+	}
 
 	return Result;
 }
diff --git a/Engine/Source/Editor/UMGEditor/Private/WidgetBlueprint.h b/Engine/Source/Editor/UMGEditor/Private/WidgetBlueprint.h
--- a/Engine/Source/Editor/UMGEditor/Private/WidgetBlueprint.h
+++ b/Engine/Source/Editor/UMGEditor/Private/WidgetBlueprint.h
@@ -64,6 +64,8 @@ public:
 	// UBlueprint interface
 	virtual UClass* GetBlueprintClass() const override;
 
+	virtual bool AllowsDynamicBinding() const override;
+
 	virtual bool SupportedByDefaultBlueprintFactory() const override
 	{
 		return false;
diff --git a/Engine/Source/Editor/UMGEditor/Private/WidgetBlueprintCompiler.cpp b/Engine/Source/Editor/UMGEditor/Private/WidgetBlueprintCompiler.cpp
--- a/Engine/Source/Editor/UMGEditor/Private/WidgetBlueprintCompiler.cpp
+++ b/Engine/Source/Editor/UMGEditor/Private/WidgetBlueprintCompiler.cpp
@@ -578,9 +578,9 @@ void FWidgetBlueprintCompiler::CleanAndSanitizeClass(UBlueprintGeneratedClass* C
 	check(ClassToClean == NewClass);
 	NewWidgetBlueprintClass = CastChecked<UWidgetBlueprintGeneratedClass>((UObject*)NewClass);
 
-	// Purge all subobjects (properties, functions, params) of the class, as they will be regenerated
-	TArray<UObject*> ClassSubObjects;
-	GetObjectsWithOuter(ClassToClean, ClassSubObjects, true);
+	//// Purge all subobjects (properties, functions, params) of the class, as they will be regenerated
+	//TArray<UObject*> ClassSubObjects;
+	//GetObjectsWithOuter(ClassToClean, ClassSubObjects, true);
 
 	//UWidgetBlueprint* Blueprint = WidgetBlueprint();
 
@@ -590,6 +590,19 @@ void FWidgetBlueprintCompiler::CleanAndSanitizeClass(UBlueprintGeneratedClass* C
 	//}
 }
 
+void FWidgetBlueprintCompiler::SaveSubObjectsFromCleanAndSanitizeClass(FSubobjectCollection& SubObjectsToSave, UBlueprintGeneratedClass* ClassToClean, UObject*& OldCDO)
+{
+	// Make sure our typed pointer is set
+	check(ClassToClean == NewClass);
+	NewWidgetBlueprintClass = CastChecked<UWidgetBlueprintGeneratedClass>((UObject*)NewClass);
+
+	UWidgetBlueprint* Blueprint = WidgetBlueprint();
+
+	SubObjectsToSave.AddObject(Blueprint->WidgetTree);
+
+	SubObjectsToSave.AddObject( NewWidgetBlueprintClass->WidgetTree );
+}
+
 void FWidgetBlueprintCompiler::CreateClassVariablesFromBlueprint()
 {
 	Super::CreateClassVariablesFromBlueprint();
@@ -629,14 +642,19 @@ void FWidgetBlueprintCompiler::FinishCompilingClass(UClass* Class)
 	UWidgetBlueprint* Blueprint = WidgetBlueprint();
 
 	UWidgetBlueprintGeneratedClass* BPGClass = CastChecked<UWidgetBlueprintGeneratedClass>(Class);
-	BPGClass->WidgetTree = NULL;
 
-	// @todo UMG Possibly need duplication here 
-	BPGClass->WidgetTree = Blueprint->WidgetTree;
-	BPGClass->AnimationData = Blueprint->AnimationData;
+	BPGClass->WidgetTree = DuplicateObject<UWidgetTree>(Blueprint->WidgetTree, BPGClass);
+
+	BPGClass->AnimationData.Reset();
+	for ( UMovieScene* Sequence : Blueprint->AnimationData )
+	{
+		UMovieScene* ClonedSequence = DuplicateObject<UMovieScene>(Sequence, BPGClass);
+		BPGClass->AnimationData.Add( ClonedSequence );
+	}
 
 	BPGClass->Bindings.Reset();
 
+	// @todo UMG Possibly need duplication here 
 	for ( const FDelegateEditorBinding& EditorBinding : Blueprint->Bindings )
 	{
 		BPGClass->Bindings.Add(EditorBinding.ToRuntimeBinding(Blueprint));
diff --git a/Engine/Source/Editor/UMGEditor/Private/WidgetBlueprintCompiler.h b/Engine/Source/Editor/UMGEditor/Private/WidgetBlueprintCompiler.h
--- a/Engine/Source/Editor/UMGEditor/Private/WidgetBlueprintCompiler.h
+++ b/Engine/Source/Editor/UMGEditor/Private/WidgetBlueprintCompiler.h
@@ -30,6 +30,7 @@ protected:
 	virtual void MergeUbergraphPagesIn(UEdGraph* Ubergraph) override;
 	virtual void SpawnNewClass(const FString& NewClassName) override;
 	virtual void CleanAndSanitizeClass(UBlueprintGeneratedClass* ClassToClean, UObject*& OldCDO) override;
+	virtual void SaveSubObjectsFromCleanAndSanitizeClass(FSubobjectCollection& SubObjectsToSave, UBlueprintGeneratedClass* ClassToClean, UObject*& OldCDO) override;
 	virtual void EnsureProperGeneratedClass(UClass*& TargetClass) override;
 	virtual void CreateClassVariablesFromBlueprint() override;
 	virtual void FinishCompilingClass(UClass* Class) override;
diff --git a/Engine/Source/Editor/UMGEditor/Private/WidgetBlueprintEditor.cpp b/Engine/Source/Editor/UMGEditor/Private/WidgetBlueprintEditor.cpp
--- a/Engine/Source/Editor/UMGEditor/Private/WidgetBlueprintEditor.cpp
+++ b/Engine/Source/Editor/UMGEditor/Private/WidgetBlueprintEditor.cpp
@@ -375,7 +375,11 @@ void FWidgetBlueprintEditor::UpdatePreview(UBlueprint* InBlueprint, bool bInForc
 		DestroyPreview();
 
 		// Save the Blueprint we're creating a preview for
-		PreviewBlueprint = InBlueprint;
+		PreviewBlueprint = Cast<UWidgetBlueprint>(InBlueprint);
+
+		// Update the generated class'es widget tree to match the blueprint tree.  That way the preview can update
+		// without needing to do a full recompile.
+		Cast<UWidgetBlueprintGeneratedClass>(PreviewBlueprint->GeneratedClass)->WidgetTree = DuplicateObject<UWidgetTree>(PreviewBlueprint->WidgetTree, PreviewBlueprint->GeneratedClass);
 
 		PreviewActor = ConstructObject<UUserWidget>(PreviewBlueprint->GeneratedClass, PreviewScene.GetWorld()->GetCurrentLevel());
 		PreviewActor->SetFlags(RF_Transactional);
diff --git a/Engine/Source/Editor/UMGEditor/Private/WidgetBlueprintEditor.h b/Engine/Source/Editor/UMGEditor/Private/WidgetBlueprintEditor.h
--- a/Engine/Source/Editor/UMGEditor/Private/WidgetBlueprintEditor.h
+++ b/Engine/Source/Editor/UMGEditor/Private/WidgetBlueprintEditor.h
@@ -103,7 +103,7 @@ private:
 	TSharedPtr<ISequencer> Sequencer;
 
 	/** The Blueprint associated with the current preview */
-	UBlueprint* PreviewBlueprint;
+	UWidgetBlueprint* PreviewBlueprint;
 
 	/** Default movie scene for new animations */
 	UMovieScene* DefaultMovieScene;
diff --git a/Engine/Source/Runtime/Engine/Classes/Engine/Blueprint.h b/Engine/Source/Runtime/Engine/Classes/Engine/Blueprint.h
--- a/Engine/Source/Runtime/Engine/Classes/Engine/Blueprint.h
+++ b/Engine/Source/Runtime/Engine/Classes/Engine/Blueprint.h
@@ -498,6 +498,9 @@ public:
 	/** @return the user-friendly name of the blueprint */
 	virtual FString GetFriendlyName() const;
 
+	/** @return true if the blueprint supports event binding for multicast delegates */
+	virtual bool AllowsDynamicBinding() const;
+
 	bool ChangeOwnerOfTemplates();
 
 #endif	//#if WITH_EDITOR
diff --git a/Engine/Source/Runtime/Engine/Classes/Engine/BlueprintGeneratedClass.h b/Engine/Source/Runtime/Engine/Classes/Engine/BlueprintGeneratedClass.h
--- a/Engine/Source/Runtime/Engine/Classes/Engine/BlueprintGeneratedClass.h
+++ b/Engine/Source/Runtime/Engine/Classes/Engine/BlueprintGeneratedClass.h
@@ -444,7 +444,7 @@ public:
 #endif
 
 	/** Bind functions on supplied actor to delegates */
-	void BindDynamicDelegates(AActor* InInstance);
+	void BindDynamicDelegates(UObject* InInstance) const;
 
 	// Finds the desired dynamic binding object for this blueprint generated class
 	UDynamicBlueprintBinding* GetDynamicBindingObject(UClass* Class) const;
diff --git a/Engine/Source/Runtime/Engine/Classes/Engine/ComponentDelegateBinding.h b/Engine/Source/Runtime/Engine/Classes/Engine/ComponentDelegateBinding.h
--- a/Engine/Source/Runtime/Engine/Classes/Engine/ComponentDelegateBinding.h
+++ b/Engine/Source/Runtime/Engine/Classes/Engine/ComponentDelegateBinding.h
@@ -38,6 +38,6 @@ class ENGINE_API UComponentDelegateBinding : public UDynamicBlueprintBinding
 	TArray<FBlueprintComponentDelegateBinding> ComponentDelegateBindings;
 
 	// Begin DynamicBlueprintBinding interface
-	virtual void BindDynamicDelegates(AActor* InInstance) const override;
+	virtual void BindDynamicDelegates(UObject* InInstance) const override;
 	// End DynamicBlueprintBinding interface
 };
diff --git a/Engine/Source/Runtime/Engine/Classes/Engine/DynamicBlueprintBinding.h b/Engine/Source/Runtime/Engine/Classes/Engine/DynamicBlueprintBinding.h
--- a/Engine/Source/Runtime/Engine/Classes/Engine/DynamicBlueprintBinding.h
+++ b/Engine/Source/Runtime/Engine/Classes/Engine/DynamicBlueprintBinding.h
@@ -9,5 +9,5 @@ class ENGINE_API UDynamicBlueprintBinding : public UObject
 {
 	GENERATED_UCLASS_BODY()
 
-	virtual void BindDynamicDelegates(AActor* InInstance) const { }
+	virtual void BindDynamicDelegates(UObject* InInstance) const { }
 };
diff --git a/Engine/Source/Runtime/Engine/Private/Blueprint.cpp b/Engine/Source/Runtime/Engine/Private/Blueprint.cpp
--- a/Engine/Source/Runtime/Engine/Private/Blueprint.cpp
+++ b/Engine/Source/Runtime/Engine/Private/Blueprint.cpp
@@ -571,6 +571,11 @@ FString UBlueprint::GetFriendlyName() const
 	return GetName();
 }
 
+bool UBlueprint::AllowsDynamicBinding() const
+{
+	return FBlueprintEditorUtils::IsActorBased(this);
+}
+
 struct FBlueprintInnerHelper
 {
 	template<typename TOBJ, typename TARR>
diff --git a/Engine/Source/Runtime/Engine/Private/BlueprintGeneratedClass.cpp b/Engine/Source/Runtime/Engine/Private/BlueprintGeneratedClass.cpp
--- a/Engine/Source/Runtime/Engine/Private/BlueprintGeneratedClass.cpp
+++ b/Engine/Source/Runtime/Engine/Private/BlueprintGeneratedClass.cpp
@@ -219,8 +219,7 @@ UDynamicBlueprintBinding* UBlueprintGeneratedClass::GetDynamicBindingObject(UCla
 	return DynamicBindingObject;
 }
 
-
-void UBlueprintGeneratedClass::BindDynamicDelegates(AActor* InInstance)
+void UBlueprintGeneratedClass::BindDynamicDelegates(UObject* InInstance) const
 {
 	if(!InInstance->IsA(this))
 	{
@@ -232,8 +231,8 @@ void UBlueprintGeneratedClass::BindDynamicDelegates(AActor* InInstance)
 	{
 		if ( ensure(DynamicBindingObjects[Index] != NULL) )
 		{
-		DynamicBindingObjects[Index]->BindDynamicDelegates(InInstance);
-	}
+			DynamicBindingObjects[Index]->BindDynamicDelegates(InInstance);
+		}
 	}
 
 	// call on super class, if it's a BlueprintGeneratedClass
@@ -262,7 +261,7 @@ bool UBlueprintGeneratedClass::GetGeneratedClassesHierarchy(const UClass* InClas
 
 UActorComponent* UBlueprintGeneratedClass::FindComponentTemplateByName(const FName& TemplateName)
 {
-	for(int32 i=0; i<ComponentTemplates.Num(); i++)
+	for(int32 i = 0; i < ComponentTemplates.Num(); i++)
 	{
 		UActorComponent* Template = ComponentTemplates[i];
 		if(Template != NULL && Template->GetFName() == TemplateName)
diff --git a/Engine/Source/Runtime/Engine/Private/ComponentDelegateBinding.cpp b/Engine/Source/Runtime/Engine/Private/ComponentDelegateBinding.cpp
--- a/Engine/Source/Runtime/Engine/Private/ComponentDelegateBinding.cpp
+++ b/Engine/Source/Runtime/Engine/Private/ComponentDelegateBinding.cpp
@@ -9,7 +9,7 @@ UComponentDelegateBinding::UComponentDelegateBinding(const class FPostConstructI
 {
 }
 
-void UComponentDelegateBinding::BindDynamicDelegates(AActor* InInstance) const
+void UComponentDelegateBinding::BindDynamicDelegates(UObject* InInstance) const
 {
 	for(int32 BindIdx=0; BindIdx<ComponentDelegateBindings.Num(); BindIdx++)
 	{
@@ -22,7 +22,7 @@ void UComponentDelegateBinding::BindDynamicDelegates(AActor* InInstance) const
 		if(ObjProp != NULL && FunctionToBind != NULL)
 		{
 			// ..see if there is actually a component assigned
-			UActorComponent* Component = Cast<UActorComponent>(ObjProp->GetObjectPropertyValue_InContainer(InInstance));
+			UObject* Component = Cast<UObject>(ObjProp->GetObjectPropertyValue_InContainer(InInstance));
 			if(Component != NULL)
 			{
 				// If there is, find the delegate property on it
diff --git a/Engine/Source/Runtime/UMG/Private/Slate/SObjectWidget.cpp b/Engine/Source/Runtime/UMG/Private/Slate/SObjectWidget.cpp
--- a/Engine/Source/Runtime/UMG/Private/Slate/SObjectWidget.cpp
+++ b/Engine/Source/Runtime/UMG/Private/Slate/SObjectWidget.cpp
@@ -23,7 +23,10 @@ void SObjectWidget::AddReferencedObjects(FReferenceCollector& Collector)
 
 void SObjectWidget::Tick(const FGeometry& AllottedGeometry, const double InCurrentTime, const float InDeltaTime)
 {
-	return WidgetObject->Tick(AllottedGeometry, InDeltaTime);
+	if ( WidgetObject )
+	{
+		return WidgetObject->Tick(AllottedGeometry, InDeltaTime);
+	}
 }
 
 //bool OnHitTest(const FGeometry& MyGeometry, FVector2D InAbsoluteCursorPosition)
diff --git a/Engine/Source/Runtime/UMG/Private/WidgetBlueprintGeneratedClass.cpp b/Engine/Source/Runtime/UMG/Private/WidgetBlueprintGeneratedClass.cpp
--- a/Engine/Source/Runtime/UMG/Private/WidgetBlueprintGeneratedClass.cpp
+++ b/Engine/Source/Runtime/UMG/Private/WidgetBlueprintGeneratedClass.cpp
@@ -16,14 +16,14 @@ void UWidgetBlueprintGeneratedClass::PostInitProperties()
 {
 	Super::PostInitProperties();
 
-	// Create a widget tree if one doesn't already exist.
-	if ( WidgetTree == NULL )
-	{
-		WidgetTree = ConstructObject<UWidgetTree>(UWidgetTree::StaticClass(), this);
-	}
-
-	//WidgetTree->SetFlags(RF_DefaultSubObject);
-	WidgetTree->SetFlags(RF_Transactional);
+	//// Create a widget tree if one doesn't already exist.
+	//if ( WidgetTree == NULL )
+	//{
+	//	WidgetTree = ConstructObject<UWidgetTree>(UWidgetTree::StaticClass(), this);
+	//}
+
+	////WidgetTree->SetFlags(RF_DefaultSubObject);
+	//WidgetTree->SetFlags(RF_Transactional);
 }
 
 void UWidgetBlueprintGeneratedClass::Link(FArchive& Ar, bool bRelinkExistingProperties)
@@ -49,67 +49,73 @@ void UWidgetBlueprintGeneratedClass::Link(FArchive& Ar, bool bRelinkExistingProp
 void UWidgetBlueprintGeneratedClass::InitializeWidget(UUserWidget* UserWidget) const
 {
 	UWidgetTree* ClonedTree = DuplicateObject<UWidgetTree>( WidgetTree, UserWidget );
+	//check(ClonedTree);
 
-	UserWidget->WidgetTree = ClonedTree;
+	if ( ClonedTree )
+	{
+		ClonedTree->SetFlags(RF_Transactional);
 
-	UClass* ActorClass = UserWidget->GetClass();
+		UserWidget->WidgetTree = ClonedTree;
 
-	TArray<UWidget*> ClonedWidgets;
-	ClonedTree->GetAllWidgets(ClonedWidgets);
+		UClass* WidgetBlueprintClass = UserWidget->GetClass();
 
-	for ( UWidget* Widget : ClonedWidgets )
-	{
-		// Not fatal if NULL, but shouldn't happen
-		if ( !ensure(Widget != NULL) )
-		{
-			continue;
-		}
+		TArray<UWidget*> ClonedWidgets;
+		ClonedTree->GetAllWidgets(ClonedWidgets);
 
-		//FName NewName = *( FString::Printf(TEXT("WidgetComponent__%d"), UserWidget->Components.Num()) );
+		for ( UWidget* Widget : ClonedWidgets )
+		{
+			// Not fatal if NULL, but shouldn't happen
+			if ( !ensure(Widget != NULL) )
+			{
+				continue;
+			}
 
-		FString VariableName = Widget->GetName();
+			FString VariableName = Widget->GetName();
 
-		Widget->bCreatedByConstructionScript = true; // Indicate it comes from a blueprint so it gets cleared when we rerun construction scripts
-		UserWidget->Components.Add(Widget); // Add to array so it gets saved
-//		Widget->SetNetAddressable();	// This component has a stable name that can be referenced for replication
+			Widget->bCreatedByConstructionScript = true; // Indicate it comes from a blueprint so it gets cleared when we rerun construction scripts
+			UserWidget->Components.Add(Widget); // Add to array so it gets saved
 
-		// Find property with the same name as the template and assign the new Timeline to it
-		UObjectPropertyBase* Prop = FindField<UObjectPropertyBase>(ActorClass, *VariableName);
-		if ( Prop )
-		{
-			Prop->SetObjectPropertyValue_InContainer(UserWidget, Widget);
-		}
+			// Find property with the same name as the template and assign the new widget to it.
+			UObjectPropertyBase* Prop = FindField<UObjectPropertyBase>(WidgetBlueprintClass, *VariableName);
+			if ( Prop )
+			{
+				Prop->SetObjectPropertyValue_InContainer(UserWidget, Widget);
+			}
 
-		// Perform binding
-		for ( const FDelegateRuntimeBinding& Binding : Bindings )
-		{
-			//TODO UMG Terrible performance, improve with Maps.
-			if ( Binding.ObjectName == VariableName )
+			// Perform binding
+			for ( const FDelegateRuntimeBinding& Binding : Bindings )
 			{
-				UFunction* BoundFunction = UserWidget->FindFunction(Binding.FunctionName);
+				//TODO UMG Terrible performance, improve with Maps.
+				if ( Binding.ObjectName == VariableName )
+				{
+					UFunction* BoundFunction = UserWidget->FindFunction(Binding.FunctionName);
 
-				FString DelegateName = Binding.PropertyName.ToString() + "Delegate";
+					FString DelegateName = Binding.PropertyName.ToString() + "Delegate";
 
-				for ( TFieldIterator<UProperty> It(Widget->GetClass()); It; ++It )
-				{
-					if ( UDelegateProperty* DelegateProp = Cast<UDelegateProperty>(*It) )
+					for ( TFieldIterator<UProperty> It(Widget->GetClass()); It; ++It )
 					{
-						if ( DelegateProp->GetName() == DelegateName || DelegateProp->GetFName() == Binding.PropertyName )
+						if ( UDelegateProperty* DelegateProp = Cast<UDelegateProperty>(*It) )
 						{
-							FScriptDelegate* ScriptDelegate = DelegateProp->GetPropertyValuePtr_InContainer(Widget);
-							ScriptDelegate->BindUFunction(UserWidget, Binding.FunctionName);
-							break;
+							if ( DelegateProp->GetName() == DelegateName || DelegateProp->GetFName() == Binding.PropertyName )
+							{
+								FScriptDelegate* ScriptDelegate = DelegateProp->GetPropertyValuePtr_InContainer(Widget);
+								ScriptDelegate->BindUFunction(UserWidget, Binding.FunctionName);
+								break;
+							}
 						}
 					}
 				}
 			}
+
+	#if WITH_EDITOR
+			Widget->ConnectEditorData();
+	#endif
 		}
 
-//		Widget->RegisterComponent();
+		// Bind any delegates on widgets
+		BindDynamicDelegates(UserWidget);
 
-#if WITH_EDITOR
-		Widget->ConnectEditorData();
-#endif
+		//TODO Add OnWidgetInitialized
 	}
 }
 
diff --git a/Engine/Source/Runtime/UMG/Public/Blueprint/WidgetBlueprintGeneratedClass.h b/Engine/Source/Runtime/UMG/Public/Blueprint/WidgetBlueprintGeneratedClass.h
--- a/Engine/Source/Runtime/UMG/Public/Blueprint/WidgetBlueprintGeneratedClass.h
+++ b/Engine/Source/Runtime/UMG/Public/Blueprint/WidgetBlueprintGeneratedClass.h
@@ -27,6 +27,12 @@ struct FDelegateRuntimeBinding
 
 class UMovieScene;
 
+/**
+ * The widget blueprint generated class allows us to create blueprintable widgets for UMG at runtime.
+ * All WBPGC's are of UUserWidget classes, and they perform special post initialization using this class
+ * to give themselves many of the same capabiltities as AActor blueprints, like dynamic delegate binding for
+ * widgets.
+ */
 UCLASS()
 class UMG_API UWidgetBlueprintGeneratedClass : public UBlueprintGeneratedClass
 {
@@ -49,5 +55,10 @@ class UMG_API UWidgetBlueprintGeneratedClass : public UBlueprintGeneratedClass
 
 	virtual void Link(FArchive& Ar, bool bRelinkExistingProperties) override;
 
-	void InitializeWidget(class UUserWidget* Actor) const;
+	/**
+	 * This is the function that makes UMG work.  Once a user widget is constructed, it will post load
+	 * call into its generated class and ask to be initialized.  The class will perform all the delegate
+	 * binding and wiring nessesary to have the user's widget perform as desired.
+	 */
+	void InitializeWidget(class UUserWidget* UserWidget) const;
 };
\ No newline at end of file
diff --git a/Engine/Source/Runtime/UMG/Public/Blueprint/WidgetTree.h b/Engine/Source/Runtime/UMG/Public/Blueprint/WidgetTree.h
--- a/Engine/Source/Runtime/UMG/Public/Blueprint/WidgetTree.h
+++ b/Engine/Source/Runtime/UMG/Public/Blueprint/WidgetTree.h
@@ -15,6 +15,9 @@ public:
 	UPROPERTY()
 	UWidget* RootWidget;
 
+	UPROPERTY()
+	TArray< UWidget* > AllWidgets;
+
 public:
 
 	UWidget* FindWidget(const FString& Name) const;
@@ -41,4 +44,13 @@ public:
 
 		return (T*)Widget;
 	}
+
+	virtual void PreSave() override
+	{
+		AllWidgets.Empty();
+
+		GetAllWidgets(AllWidgets);
+
+		Super::PreSave();
+	}
 };
