#UE4 fix for actors (and other 'Object' types) that were exposed on spawn not being assignable.. Introduced by keeping type information on pins around so that we could assign an enum value, instead we allow the erasure to occur and generate an extra node in the case of an enum...

diff --git a/Engine/Source/Editor/KismetCompiler/Private/KismetCompilerMisc.cpp b/Engine/Source/Editor/KismetCompiler/Private/KismetCompilerMisc.cpp
--- a/Engine/Source/Editor/KismetCompiler/Private/KismetCompilerMisc.cpp
+++ b/Engine/Source/Editor/KismetCompiler/Private/KismetCompilerMisc.cpp
@@ -6,6 +6,7 @@
 
 #include "KismetCompilerPrivatePCH.h"
 #include "KismetCompilerMisc.h"
+#include "K2Node_EnumLiteral.h"
 #include "Kismet2/KismetReinstanceUtilities.h"
 #include "Editor/UnrealEd/Public/Kismet2/BlueprintEditorUtils.h"
 #include "Editor/UnrealEd/Public/Kismet2/KismetDebugUtilities.h"
@@ -546,13 +547,25 @@ UEdGraphPin* FKismetCompilerUtilities::GenerateAssignmentNodes(class FKismetComp
 				UEdGraphPin* PropertyNamePin = SetVarNode->FindPinChecked(PropertyNameParamName);
 				PropertyNamePin->DefaultValue = OrgPin->PinName;
 
-				// Move connection from the variable pin on the spawn node to the 'value' pin
 				UEdGraphPin* ValuePin = SetVarNode->FindPinChecked(ValueParamName);
-				CompilerContext.MovePinLinksToIntermediate(*OrgPin, *ValuePin);
-				SetVarNode->PinConnectionListChanged(ValuePin);
-
-				// PinSubCategoryObject is lost when we go from function back to pin when we have type information that is BP only (IE an enum created in the editor):
-				ValuePin->PinType.PinSubCategoryObject = OrgPin->PinType.PinSubCategoryObject;
+				if (OrgPin->LinkedTo.Num() == 0 &&
+					OrgPin->DefaultValue != FString() &&
+					OrgPin->PinType.PinCategory == Schema->PC_Byte &&
+					OrgPin->PinType.PinSubCategoryObject.IsValid() &&
+					OrgPin->PinType.PinSubCategoryObject->IsA<UEnum>())
+				{
+					// Pin is an enum, we need to alias the enum value to an int:
+					UK2Node_EnumLiteral* EnumLiteralNode = CompilerContext.SpawnIntermediateNode<UK2Node_EnumLiteral>(SpawnNode, SourceGraph);
+					EnumLiteralNode->Enum = CastChecked<UEnum>(OrgPin->PinType.PinSubCategoryObject.Get());
+					EnumLiteralNode->AllocateDefaultPins();
+					EnumLiteralNode->FindPinChecked(Schema->PN_ReturnValue)->MakeLinkTo(ValuePin);
+				}
+				else
+				{
+					// Move connection from the variable pin on the spawn node to the 'value' pin
+					CompilerContext.MovePinLinksToIntermediate(*OrgPin, *ValuePin);
+					SetVarNode->PinConnectionListChanged(ValuePin);
+				}
 			}
 		}
 	}
