Kismet node/pin validation - more robust code.

ttp#333960 BP: CRASH: Possible to get a crash with undo/redo after changing variable nodes

diff --git a/Engine/Source/Editor/BlueprintGraph/Private/K2Node_CallFunction.cpp b/Engine/Source/Editor/BlueprintGraph/Private/K2Node_CallFunction.cpp
--- a/Engine/Source/Editor/BlueprintGraph/Private/K2Node_CallFunction.cpp
+++ b/Engine/Source/Editor/BlueprintGraph/Private/K2Node_CallFunction.cpp
@@ -540,23 +540,22 @@ bool UK2Node_CallFunction::CreatePinsForFunctionCall(const UFunction* Function)
 	//@TODO: Can't strictly speaking always hide the self pin for pure and static functions; it'll still be needed if the function belongs to a class not in the blueprint's class hierarchy
 	SelfPin->bHidden = ((bIsPureFunc && !bIsConstFunc) || bIsProtectedFunc || bIsStaticFunc);
 
-	if (bIsStaticFunc)
+	UBlueprint* BP = GetBlueprint();
+	ensure(BP);
+	if (bIsStaticFunc && BP)
 	{
 		// Wire up the self to the CDO of the class if it's not us
-		if (UBlueprint* BP = GetBlueprint())
+		if (!BP->SkeletonGeneratedClass->IsChildOf(FunctionOwnerClass))
 		{
-			if (!BP->SkeletonGeneratedClass->IsChildOf(FunctionOwnerClass))
-			{
-				SelfPin->DefaultObject = FunctionOwnerClass->GetDefaultObject();
-			}
+			SelfPin->DefaultObject = FunctionOwnerClass->GetDefaultObject();
 		}
 	}
 
 	// Build a list of the pins that should be hidden for this function (ones that are automagically filled in by the K2 compiler)
 	TSet<FString> PinsToHide;
-	FBlueprintEditorUtils::GetHiddenPinsForFunction(GetBlueprint(), Function, PinsToHide);
+	FBlueprintEditorUtils::GetHiddenPinsForFunction(BP, Function, PinsToHide);
 
-	const bool bShowHiddenSelfPins = ((PinsToHide.Num() > 0) && GetBlueprint()->ParentClass->HasMetaData(FBlueprintMetadata::MD_ShowHiddenSelfPins));
+	const bool bShowHiddenSelfPins = ((PinsToHide.Num() > 0) && BP && BP->ParentClass && BP->ParentClass->HasMetaData(FBlueprintMetadata::MD_ShowHiddenSelfPins));
 
 	// Create the inputs and outputs
 	bool bAllPinsGood = true;
diff --git a/Engine/Source/Editor/KismetCompiler/Private/EdGraphCompiler.cpp b/Engine/Source/Editor/KismetCompiler/Private/EdGraphCompiler.cpp
--- a/Engine/Source/Editor/KismetCompiler/Private/EdGraphCompiler.cpp
+++ b/Engine/Source/Editor/KismetCompiler/Private/EdGraphCompiler.cpp
@@ -27,13 +27,13 @@ void FGraphCompilerContext::ValidateLink(const UEdGraphPin* PinA, const UEdGraph
 /** Validate that the wiring for a single pin is schema compatible */
 void FGraphCompilerContext::ValidatePin(const UEdGraphPin* Pin) const
 {
-	if (NULL == Pin->GetOwningNodeUnchecked())
+	if (!Pin || !Pin->GetOwningNodeUnchecked())
 	{
 		MessageLog.Error(
 			*FString::Printf(
 				*NSLOCTEXT("EdGraphCompiler", "PinWrongOuterError", "Blueprint is corrupted! Pin '%s' has wrong outer '%s'.").ToString(),
-				*Pin->GetName(),
-				Pin->GetOuter() ? *Pin->GetOuter()->GetName() : TEXT("NULL")),
+				Pin ? *Pin->GetName() : TEXT("UNKNOWN"),
+				(Pin && Pin->GetOuter()) ? *Pin->GetOuter()->GetName() : TEXT("NULL")),
 			Pin);
 		return;
 	}
@@ -66,7 +66,15 @@ void FGraphCompilerContext::ValidateNode(const UEdGraphNode* Node) const
 	{
 		if (const UEdGraphPin* Pin = Node->Pins[PinIndex])
 		{
-			ValidatePin(Pin);
+			auto OwningNode = Pin->GetOwningNodeUnchecked();
+			if (OwningNode != Node)
+			{
+				MessageLog.Error(*NSLOCTEXT("EdGraphCompiler", "WrongPinsOwner_Error", "The pin @@ has outer @@, but it's used in @@").ToString(), Pin, OwningNode, Node);
+			}
+			else
+			{
+				ValidatePin(Pin);
+			}
 		}
 	}
 
diff --git a/Engine/Source/Editor/KismetCompiler/Private/KismetCompiler.cpp b/Engine/Source/Editor/KismetCompiler/Private/KismetCompiler.cpp
--- a/Engine/Source/Editor/KismetCompiler/Private/KismetCompiler.cpp
+++ b/Engine/Source/Editor/KismetCompiler/Private/KismetCompiler.cpp
@@ -279,6 +279,13 @@ void FKismetCompilerContext::ValidatePin(const UEdGraphPin* Pin) const
 {
 	Super::ValidatePin(Pin);
 
+	auto OwningNodeUnchecked = Pin ? Pin->GetOwningNodeUnchecked() : NULL;
+	if (!OwningNodeUnchecked)
+	{
+		//handled by Super::ValidatePin
+		return;
+	}
+
 	// Fixing up references to the skel or the generated classes to be PSC_Self pins
 	if ((Pin->PinType.PinCategory == Schema->PC_Object) || (Pin->PinType.PinCategory == Schema->PC_Interface))
 	{
@@ -295,7 +302,7 @@ void FKismetCompilerContext::ValidatePin(const UEdGraphPin* Pin) const
 	if (Pin->PinType.PinCategory == Schema->PC_Wildcard)
 	{
 		// Wildcard pins should never be seen by the compiler; they should always be forced into a particular type by wiring.
-		MessageLog.Error(*LOCTEXT("UndeterminedPinType_Error", "The type of @@ is undetermined.  Connect something to @@ to imply a specific type.").ToString(), Pin, Pin->GetOwningNode()); 
+		MessageLog.Error(*LOCTEXT("UndeterminedPinType_Error", "The type of @@ is undetermined.  Connect something to @@ to imply a specific type.").ToString(), Pin, OwningNodeUnchecked);
 	}
 
 	if (Pin->LinkedTo.Num() > 1)
@@ -317,7 +324,7 @@ void FKismetCompilerContext::ValidatePin(const UEdGraphPin* Pin) const
 			else if( Schema->IsSelfPin(*Pin) )
 			{
 				// Pure functions and latent functions cannot have more than one self connection
-				UK2Node_CallFunction* OwningNode = Cast<UK2Node_CallFunction>(Pin->GetOwningNode());
+				UK2Node_CallFunction* OwningNode = Cast<UK2Node_CallFunction>(OwningNodeUnchecked);
 				if( OwningNode )
 				{
 					if( OwningNode->IsNodePure() )
@@ -344,7 +351,7 @@ void FKismetCompilerContext::ValidatePin(const UEdGraphPin* Pin) const
 	//function return node exec pin should be connected to something
 	if(Pin->Direction == EGPD_Input && Pin->LinkedTo.Num() == 0 && Schema->IsExecPin(*Pin) )
 	{
-		if(UK2Node_FunctionResult* OwningNode = Cast<UK2Node_FunctionResult>(Pin->GetOwningNode()))
+		if (UK2Node_FunctionResult* OwningNode = Cast<UK2Node_FunctionResult>(OwningNodeUnchecked))
 		{
 			MessageLog.Warning(*LOCTEXT("ReturnNodeExecPinUnconnected", "ReturnNode Exec pin has no connections on @@").ToString(), Pin);
 		}
