Preventing blueprint function libraries from being marked as blueprint variable types.
(TTP #340627)

diff --git a/Engine/Source/Editor/KismetCompiler/Private/KismetCompiler.cpp b/Engine/Source/Editor/KismetCompiler/Private/KismetCompiler.cpp
--- a/Engine/Source/Editor/KismetCompiler/Private/KismetCompiler.cpp
+++ b/Engine/Source/Editor/KismetCompiler/Private/KismetCompiler.cpp
@@ -1631,7 +1631,12 @@ void FKismetCompilerContext::FinishCompilingClass(UClass* Class)
 			Class->SetMetaData(TEXT("Category"), *Blueprint->BlueprintCategory);
 		}
 
-		Class->SetMetaData(FBlueprintMetadata::MD_AllowableBlueprintVariableType, TEXT("true"));
+		if ((Blueprint->BlueprintType == BPTYPE_Normal) || 
+			(Blueprint->BlueprintType == BPTYPE_Const)  || 
+			(Blueprint->BlueprintType == BPTYPE_Interface))
+		{
+			//Class->SetMetaData(FBlueprintMetadata::MD_AllowableBlueprintVariableType, TEXT("true"));
+		}
 #endif
 
 		// Add in additional flags implied by the blueprint
