Preventing blueprint function libraries from being marked as blueprint variable types (last checkin, the file wasn't up to date).
(TTP #340627)

diff --git a/Engine/Source/Editor/KismetCompiler/Private/KismetCompiler.cpp b/Engine/Source/Editor/KismetCompiler/Private/KismetCompiler.cpp
--- a/Engine/Source/Editor/KismetCompiler/Private/KismetCompiler.cpp
+++ b/Engine/Source/Editor/KismetCompiler/Private/KismetCompiler.cpp
@@ -1635,7 +1635,7 @@ void FKismetCompilerContext::FinishCompilingClass(UClass* Class)
 			(Blueprint->BlueprintType == BPTYPE_Const)  || 
 			(Blueprint->BlueprintType == BPTYPE_Interface))
 		{
-			//Class->SetMetaData(FBlueprintMetadata::MD_AllowableBlueprintVariableType, TEXT("true"));
+			Class->SetMetaData(FBlueprintMetadata::MD_AllowableBlueprintVariableType, TEXT("true"));
 		}
 #endif
 
