Marking functions that were generated from intermediate events (events added during the ExpandNode() step) as blueprint internal (so we can ignore them for the action menu building).

diff --git a/Engine/Source/Editor/KismetCompiler/Private/KismetCompiler.cpp b/Engine/Source/Editor/KismetCompiler/Private/KismetCompiler.cpp
--- a/Engine/Source/Editor/KismetCompiler/Private/KismetCompiler.cpp
+++ b/Engine/Source/Editor/KismetCompiler/Private/KismetCompiler.cpp
@@ -1442,6 +1442,11 @@ void FKismetCompilerContext::FinishCompilingFunction(FKismetFunctionContext& Con
 		Function->FunctionFlags |= FUNC_BlueprintCallable;
 	}
 
+	if (Context.IsBlueprintInternal())
+	{
+		Function->SetMetaData(*FBlueprintMetadata::MD_BlueprintInternalUseOnly.ToString(), TEXT("true"));
+	}
+
 	if (Context.IsInterfaceStub())
 	{
 		Function->FunctionFlags |= FUNC_BlueprintEvent;
@@ -2186,7 +2191,7 @@ void FKismetCompilerContext::CreateFunctionStubForEvent(UK2Node_Event* SrcEventN
 
 	if (SrcEventNode->bOverrideFunction || SrcEventNode->bInternalEvent)
 	{
-		StubContext.MarkAsInternalOrCppUseOnly();
+		StubContext.MarkAsUncallable();
 	}
 
 	if ((SrcEventNode->FunctionFlags & FUNC_Net) > 0)
@@ -2194,6 +2199,15 @@ void FKismetCompilerContext::CreateFunctionStubForEvent(UK2Node_Event* SrcEventN
 		StubContext.MarkAsNetFunction(SrcEventNode->FunctionFlags);
 	}
 
+	UPackage* OuterPackage = SrcEventNode->GetOutermost();
+	check(OuterPackage != nullptr);
+	bool bIsIntermediate = OuterPackage->GetMetaData()->HasValue(SrcEventNode, FBlueprintMetadata::MD_BlueprintInternalUseOnly);
+
+	if (bIsIntermediate)
+	{
+		StubContext.MarkAsIntermediate();
+	}
+
 	// Create an entry point
 	UK2Node_FunctionEntry* EntryNode = SpawnIntermediateNode<UK2Node_FunctionEntry>(SrcEventNode, ChildStubGraph);
 	EntryNode->NodePosX = -200;
@@ -2514,7 +2528,8 @@ void FKismetCompilerContext::CreateAndProcessUbergraph()
 		UbergraphContext = new (FunctionList) FKismetFunctionContext(MessageLog, Schema, NewClass, Blueprint);
 		UbergraphContext->SourceGraph = ConsolidatedEventGraph;
 		UbergraphContext->MarkAsEventGraph();
-		UbergraphContext->MarkAsInternalOrCppUseOnly();
+		UbergraphContext->MarkAsUncallable();
+		UbergraphContext->MarkAsIntermediate();
 		UbergraphContext->SetExternalNetNameMap(&ClassScopeNetNameMap);
 
 		Blueprint->EventGraphs.Empty();
diff --git a/Engine/Source/Editor/KismetCompiler/Private/KismetCompilerMisc.cpp b/Engine/Source/Editor/KismetCompiler/Private/KismetCompilerMisc.cpp
--- a/Engine/Source/Editor/KismetCompiler/Private/KismetCompilerMisc.cpp
+++ b/Engine/Source/Editor/KismetCompiler/Private/KismetCompilerMisc.cpp
@@ -946,6 +946,7 @@ FKismetFunctionContext::FKismetFunctionContext(FCompilerResultsLog& InMessageLog
 	, Schema(InSchema)
 	, bIsUbergraph(false)
 	, bCannotBeCalledFromOtherKismet(false)
+	, bBlueprintInternal(false)
 	, bIsInterfaceStub(false)
 	, bIsConstFunction(false)
 	// only need debug-data when running in the editor app:
diff --git a/Engine/Source/Editor/KismetCompiler/Public/KismetCompiledFunctionContext.h b/Engine/Source/Editor/KismetCompiler/Public/KismetCompiledFunctionContext.h
--- a/Engine/Source/Editor/KismetCompiler/Public/KismetCompiledFunctionContext.h
+++ b/Engine/Source/Editor/KismetCompiler/Public/KismetCompiledFunctionContext.h
@@ -67,6 +67,7 @@ public:
 
 	bool bIsUbergraph;
 	bool bCannotBeCalledFromOtherKismet;
+	bool bBlueprintInternal;
 	bool bIsInterfaceStub;
 	bool bIsConstFunction;
 	bool bCreateDebugData;
@@ -105,7 +106,7 @@ public:
 		bIsUbergraph = true;
 	}
 
-	void MarkAsInternalOrCppUseOnly()
+	void MarkAsUncallable()
 	{
 		bCannotBeCalledFromOtherKismet = true;
 	}
@@ -115,6 +116,16 @@ public:
 		return !bCannotBeCalledFromOtherKismet;
 	}
 
+	void MarkAsIntermediate()
+	{
+		bBlueprintInternal = true;
+	}
+
+	bool IsBlueprintInternal() const
+	{
+		return bBlueprintInternal;
+	}
+
 	void MarkAsInterfaceStub()
 	{
 		bIsInterfaceStub = true;
diff --git a/Engine/Source/Editor/KismetCompiler/Public/KismetCompiler.h b/Engine/Source/Editor/KismetCompiler/Public/KismetCompiler.h
--- a/Engine/Source/Editor/KismetCompiler/Public/KismetCompiler.h
+++ b/Engine/Source/Editor/KismetCompiler/Public/KismetCompiler.h
@@ -110,6 +110,10 @@ public:
 
 		AutoAssignNodePosition(Result);
 
+		UPackage* OuterPackage = Result->GetOutermost();
+		check(OuterPackage != nullptr);
+		OuterPackage->GetMetaData()->SetValue(Result, *FBlueprintMetadata::MD_BlueprintInternalUseOnly.ToString(), TEXT("true"));
+
 		return Result;
 	}
 
