Back out changelist 2170880

diff --git a/Engine/Source/Editor/KismetCompiler/Private/KismetCompiler.cpp b/Engine/Source/Editor/KismetCompiler/Private/KismetCompiler.cpp
--- a/Engine/Source/Editor/KismetCompiler/Private/KismetCompiler.cpp
+++ b/Engine/Source/Editor/KismetCompiler/Private/KismetCompiler.cpp
@@ -1442,11 +1442,6 @@ void FKismetCompilerContext::FinishCompilingFunction(FKismetFunctionContext& Con
 		Function->FunctionFlags |= FUNC_BlueprintCallable;
 	}
 
-	if (Context.IsBlueprintInternal())
-	{
-		Function->SetMetaData(*FBlueprintMetadata::MD_BlueprintInternalUseOnly.ToString(), TEXT("true"));
-	}
-
 	if (Context.IsInterfaceStub())
 	{
 		Function->FunctionFlags |= FUNC_BlueprintEvent;
@@ -2191,7 +2186,7 @@ void FKismetCompilerContext::CreateFunctionStubForEvent(UK2Node_Event* SrcEventN
 
 	if (SrcEventNode->bOverrideFunction || SrcEventNode->bInternalEvent)
 	{
-		StubContext.MarkAsUncallable();
+		StubContext.MarkAsInternalOrCppUseOnly();
 	}
 
 	if ((SrcEventNode->FunctionFlags & FUNC_Net) > 0)
@@ -2199,15 +2194,6 @@ void FKismetCompilerContext::CreateFunctionStubForEvent(UK2Node_Event* SrcEventN
 		StubContext.MarkAsNetFunction(SrcEventNode->FunctionFlags);
 	}
 
-	UPackage* OuterPackage = SrcEventNode->GetOutermost();
-	check(OuterPackage != nullptr);
-	bool bIsIntermediate = OuterPackage->GetMetaData()->HasValue(SrcEventNode, FBlueprintMetadata::MD_BlueprintInternalUseOnly);
-
-	if (bIsIntermediate)
-	{
-		StubContext.MarkAsIntermediate();
-	}
-
 	// Create an entry point
 	UK2Node_FunctionEntry* EntryNode = SpawnIntermediateNode<UK2Node_FunctionEntry>(SrcEventNode, ChildStubGraph);
 	EntryNode->NodePosX = -200;
@@ -2528,8 +2514,7 @@ void FKismetCompilerContext::CreateAndProcessUbergraph()
 		UbergraphContext = new (FunctionList) FKismetFunctionContext(MessageLog, Schema, NewClass, Blueprint);
 		UbergraphContext->SourceGraph = ConsolidatedEventGraph;
 		UbergraphContext->MarkAsEventGraph();
-		UbergraphContext->MarkAsUncallable();
-		UbergraphContext->MarkAsIntermediate();
+		UbergraphContext->MarkAsInternalOrCppUseOnly();
 		UbergraphContext->SetExternalNetNameMap(&ClassScopeNetNameMap);
 
 		Blueprint->EventGraphs.Empty();
diff --git a/Engine/Source/Editor/KismetCompiler/Private/KismetCompilerMisc.cpp b/Engine/Source/Editor/KismetCompiler/Private/KismetCompilerMisc.cpp
--- a/Engine/Source/Editor/KismetCompiler/Private/KismetCompilerMisc.cpp
+++ b/Engine/Source/Editor/KismetCompiler/Private/KismetCompilerMisc.cpp
@@ -946,7 +946,6 @@ FKismetFunctionContext::FKismetFunctionContext(FCompilerResultsLog& InMessageLog
 	, Schema(InSchema)
 	, bIsUbergraph(false)
 	, bCannotBeCalledFromOtherKismet(false)
-	, bBlueprintInternal(false)
 	, bIsInterfaceStub(false)
 	, bIsConstFunction(false)
 	// only need debug-data when running in the editor app:
diff --git a/Engine/Source/Editor/KismetCompiler/Public/KismetCompiledFunctionContext.h b/Engine/Source/Editor/KismetCompiler/Public/KismetCompiledFunctionContext.h
--- a/Engine/Source/Editor/KismetCompiler/Public/KismetCompiledFunctionContext.h
+++ b/Engine/Source/Editor/KismetCompiler/Public/KismetCompiledFunctionContext.h
@@ -67,7 +67,6 @@ public:
 
 	bool bIsUbergraph;
 	bool bCannotBeCalledFromOtherKismet;
-	bool bBlueprintInternal;
 	bool bIsInterfaceStub;
 	bool bIsConstFunction;
 	bool bCreateDebugData;
@@ -106,7 +105,7 @@ public:
 		bIsUbergraph = true;
 	}
 
-	void MarkAsUncallable()
+	void MarkAsInternalOrCppUseOnly()
 	{
 		bCannotBeCalledFromOtherKismet = true;
 	}
@@ -116,16 +115,6 @@ public:
 		return !bCannotBeCalledFromOtherKismet;
 	}
 
-	void MarkAsIntermediate()
-	{
-		bBlueprintInternal = true;
-	}
-
-	bool IsBlueprintInternal() const
-	{
-		return bBlueprintInternal;
-	}
-
 	void MarkAsInterfaceStub()
 	{
 		bIsInterfaceStub = true;
diff --git a/Engine/Source/Editor/KismetCompiler/Public/KismetCompiler.h b/Engine/Source/Editor/KismetCompiler/Public/KismetCompiler.h
--- a/Engine/Source/Editor/KismetCompiler/Public/KismetCompiler.h
+++ b/Engine/Source/Editor/KismetCompiler/Public/KismetCompiler.h
@@ -110,10 +110,6 @@ public:
 
 		AutoAssignNodePosition(Result);
 
-		UPackage* OuterPackage = Result->GetOutermost();
-		check(OuterPackage != nullptr);
-		OuterPackage->GetMetaData()->SetValue(Result, *FBlueprintMetadata::MD_BlueprintInternalUseOnly.ToString(), TEXT("true"));
-
 		return Result;
 	}
 
