BP:  Need to expose "HideCategories" metadata in BP Editor ttp338142

diff --git a/Engine/Source/Editor/KismetCompiler/Private/KismetCompiler.cpp b/Engine/Source/Editor/KismetCompiler/Private/KismetCompiler.cpp
--- a/Engine/Source/Editor/KismetCompiler/Private/KismetCompiler.cpp
+++ b/Engine/Source/Editor/KismetCompiler/Private/KismetCompiler.cpp
@@ -1589,6 +1589,8 @@ void FKismetCompilerContext::FinishCompilingClass(UClass* Class)
 {
 	UClass* ParentClass = Class->GetSuperClass();
 
+	TArray<FString> AllHideCategories;
+
 	if (ParentClass != NULL)
 	{
 		// Propagate the new parent's inheritable class flags
@@ -1604,10 +1606,7 @@ void FKismetCompilerContext::FinishCompilingClass(UClass* Class)
 
 		// Copy the category info from the parent class
 #if WITH_EDITORONLY_DATA
-		if (ParentClass->HasMetaData(TEXT("HideCategories")))
-		{
-			Class->SetMetaData(TEXT("HideCategories"), *ParentClass->GetMetaData("HideCategories"));
-		}
+		ParentClass->GetHideCategories(AllHideCategories);
 		if (ParentClass->HasMetaData(TEXT("ShowCategories")))
 		{
 			Class->SetMetaData(TEXT("ShowCategories"), *ParentClass->GetMetaData("ShowCategories"));
@@ -1653,6 +1652,12 @@ void FKismetCompilerContext::FinishCompilingClass(UClass* Class)
 		//@TODO: Might want to be able to specify some of these here too
 	}
 
+	AllHideCategories.Append(Blueprint->HideCategories);
+	if (AllHideCategories.Num())
+	{
+		Class->SetMetaData(TEXT("HideCategories"), *FString::Join(AllHideCategories, TEXT(" ")));
+	}
+
 	// Add in any other needed flags
 	Class->ClassFlags |= (CLASS_Parsed | CLASS_CompiledFromBlueprint);
 
diff --git a/Engine/Source/Runtime/Engine/Classes/Engine/Blueprint.h b/Engine/Source/Runtime/Engine/Classes/Engine/Blueprint.h
--- a/Engine/Source/Runtime/Engine/Classes/Engine/Blueprint.h
+++ b/Engine/Source/Runtime/Engine/Classes/Engine/Blueprint.h
@@ -299,6 +299,10 @@ class ENGINE_API UBlueprint : public UBlueprintCore
 	UPROPERTY(EditAnywhere, Category=BlueprintOption)
 	FString BlueprintCategory;
 
+	/** Additional HideCategories. The are added to HideCategories from parent. */
+	UPROPERTY(EditAnywhere, Category=BlueprintOption)
+	TArray<FString> HideCategories;
+
 	/** TRUE to show a warning when attempting to start in PIE and there is a compiler error on this Blueprint */
 	UPROPERTY(transient)
 	bool bDisplayCompilePIEWarning;
