#ue4 fix for user specified values on non-default enums not making it through compilation (pin wasn't being set - we were always the first value in the enum). Reviewed by [REDACTED_PERSON]

diff --git a/Engine/Source/Editor/BlueprintGraph/Classes/K2Node_EnumLiteral.h b/Engine/Source/Editor/BlueprintGraph/Classes/K2Node_EnumLiteral.h
--- a/Engine/Source/Editor/BlueprintGraph/Classes/K2Node_EnumLiteral.h
+++ b/Engine/Source/Editor/BlueprintGraph/Classes/K2Node_EnumLiteral.h
@@ -13,7 +13,7 @@ class UK2Node_EnumLiteral : public UK2Node, public INodeDependingOnEnumInterface
 	UPROPERTY()
 	UEnum* Enum;
 
-	static const FString& GetEnumInputPinName();
+	static BLUEPRINTGRAPH_API const FString& GetEnumInputPinName();
 
 	// Begin UEdGraphNode interface
 	virtual void AllocateDefaultPins() override;
diff --git a/Engine/Source/Editor/KismetCompiler/Private/KismetCompilerMisc.cpp b/Engine/Source/Editor/KismetCompiler/Private/KismetCompilerMisc.cpp
--- a/Engine/Source/Editor/KismetCompiler/Private/KismetCompilerMisc.cpp
+++ b/Engine/Source/Editor/KismetCompiler/Private/KismetCompilerMisc.cpp
@@ -559,6 +559,10 @@ UEdGraphPin* FKismetCompilerUtilities::GenerateAssignmentNodes(class FKismetComp
 					EnumLiteralNode->Enum = CastChecked<UEnum>(OrgPin->PinType.PinSubCategoryObject.Get());
 					EnumLiteralNode->AllocateDefaultPins();
 					EnumLiteralNode->FindPinChecked(Schema->PN_ReturnValue)->MakeLinkTo(ValuePin);
+
+					UEdGraphPin* InPin = EnumLiteralNode->FindPinChecked(UK2Node_EnumLiteral::GetEnumInputPinName());
+					check( InPin );
+					InPin->DefaultValue = OrgPin->DefaultValue;
 				}
 				else
 				{
