#UE4 quick fix for cook failure... generating assignment nodes if we don't have a CDO to dig default values out of...

diff --git a/Engine/Source/Editor/KismetCompiler/Private/KismetCompilerMisc.cpp b/Engine/Source/Editor/KismetCompiler/Private/KismetCompilerMisc.cpp
--- a/Engine/Source/Editor/KismetCompiler/Private/KismetCompilerMisc.cpp
+++ b/Engine/Source/Editor/KismetCompiler/Private/KismetCompilerMisc.cpp
@@ -530,13 +530,16 @@ UEdGraphPin* FKismetCompilerUtilities::GenerateAssignmentNodes(class FKismetComp
 					continue;
 				}
 
-				// We don't want to generate an assignment node unless the default value 
-				// differs from the value in the CDO:
-				FString DefaultValueAsString;
-				FBlueprintEditorUtils::PropertyValueToString(Property, (uint8*)ForClass->ClassDefaultObject, DefaultValueAsString);
-				if (DefaultValueAsString == OrgPin->DefaultValue)
+				if( ForClass->ClassDefaultObject )
 				{
-					continue;
+					// We don't want to generate an assignment node unless the default value 
+					// differs from the value in the CDO:
+					FString DefaultValueAsString;
+					FBlueprintEditorUtils::PropertyValueToString(Property, (uint8*)ForClass->ClassDefaultObject, DefaultValueAsString);
+					if (DefaultValueAsString == OrgPin->DefaultValue)
+					{
+						continue;
+					}
 				}
 			}
 
