#UE4 fix for TTP#341984 - enum default values in macros. Reasonable to address these type safety issues at a lower level, but that systemic work can be done later.

diff --git a/Engine/Source/Editor/KismetCompiler/Private/KismetCompiler.cpp b/Engine/Source/Editor/KismetCompiler/Private/KismetCompiler.cpp
--- a/Engine/Source/Editor/KismetCompiler/Private/KismetCompiler.cpp
+++ b/Engine/Source/Editor/KismetCompiler/Private/KismetCompiler.cpp
@@ -17,6 +17,7 @@
 #include "Kismet2/Kismet2NameValidators.h"
 #include "UserDefinedStructureCompilerUtils.h"
 #include "EditorCategoryUtils.h"
+#include "K2Node_EnumLiteral.h"
 
 static bool bDebugPropertyPropagation = false;
 
@@ -2695,6 +2696,22 @@ void FKismetCompilerContext::ExpandTunnelsAndMacros(UEdGraph* SourceGraph)
 						MakeArrayOut->MakeLinkTo(Pin);
 						MakeArrayNode->PinConnectionListChanged(MakeArrayOut);
 					}
+					else if (Pin->LinkedTo.Num() == 0 &&
+							Pin->DefaultValue != FString() &&
+							Pin->PinType.PinCategory == Schema->PC_Byte &&
+							Pin->PinType.PinSubCategoryObject.IsValid() &&
+							Pin->PinType.PinSubCategoryObject->IsA<UEnum>())
+					{
+						// Similarly, enums need a 'make enum' node because they decay to byte after instantiation:
+						UK2Node_EnumLiteral* EnumLiteralNode = SpawnIntermediateNode<UK2Node_EnumLiteral>(MacroInstanceNode, SourceGraph);
+						EnumLiteralNode->Enum = CastChecked<UEnum>(Pin->PinType.PinSubCategoryObject.Get());
+						EnumLiteralNode->AllocateDefaultPins();
+						EnumLiteralNode->FindPinChecked(Schema->PN_ReturnValue)->MakeLinkTo(Pin);
+
+						UEdGraphPin* InPin = EnumLiteralNode->FindPinChecked(UK2Node_EnumLiteral::GetEnumInputPinName());
+						check(InPin);
+						InPin->DefaultValue = Pin->DefaultValue;
+					}
 					// Otherwise we need to handle the pin splitting
 					else if (Pin->SubPins.Num() > 0)
 					{
