BP: Pins of current class type, won't be convert into PCS_Self.

UE4: Blueprints: Unable to clear an object variable with Set Variable ttp340878

diff --git a/Engine/Source/Editor/KismetCompiler/Private/KismetCompiler.cpp b/Engine/Source/Editor/KismetCompiler/Private/KismetCompiler.cpp
--- a/Engine/Source/Editor/KismetCompiler/Private/KismetCompiler.cpp
+++ b/Engine/Source/Editor/KismetCompiler/Private/KismetCompiler.cpp
@@ -288,19 +288,6 @@ void FKismetCompilerContext::ValidatePin(const UEdGraphPin* Pin) const
 		return;
 	}
 
-	// Fixing up references to the skel or the generated classes to be PSC_Self pins
-	if ((Pin->PinType.PinCategory == Schema->PC_Object) || (Pin->PinType.PinCategory == Schema->PC_Interface))
-	{
-		//@todo:  This is modifying the model, but is acceptable to save another prepass on the pins
-		UEdGraphPin* MutablePin = const_cast<UEdGraphPin*>(Pin);
-
-		if (NewClass && MutablePin->PinType.PinSubCategoryObject.Get() == NewClass)
-			{
-				MutablePin->PinType.PinSubCategory = Schema->PSC_Self;
-				MutablePin->PinType.PinSubCategoryObject = NULL;
-			}
-		}
-
 	if (Pin->PinType.PinCategory == Schema->PC_Wildcard)
 	{
 		// Wildcard pins should never be seen by the compiler; they should always be forced into a particular type by wiring.
