BP: compiler optimization: Do not register nets while skeletal class compilation.

diff --git a/Engine/Source/Editor/KismetCompiler/Private/KismetCompiler.cpp b/Engine/Source/Editor/KismetCompiler/Private/KismetCompiler.cpp
--- a/Engine/Source/Editor/KismetCompiler/Private/KismetCompiler.cpp
+++ b/Engine/Source/Editor/KismetCompiler/Private/KismetCompiler.cpp
@@ -1153,27 +1153,43 @@ void FKismetCompilerContext::PrecompileFunction(FKismetFunctionContext& Context)
 		Context.NewClass->AddFunctionToFunctionMap(Context.Function);
 
 		//@TODO: Prune pure functions that don't have any consumers
-
-		// Find the execution path (and make sure it has no cycles)
-		CreateExecutionSchedule(Context.SourceGraph->Nodes, Context.LinearExecutionList);
-
-		// Create any user defined variables, this must occur before registering nets so that the properties are in place
-		UField** PropertyStorageLocation = &(Context.Function->Children);
-		CreateUserDefinedLocalVariablesForFunction(Context, PropertyStorageLocation);
-
-		bool bIsPureFunction = true;
-		for (int32 NodeIndex = 0; NodeIndex < Context.LinearExecutionList.Num(); ++NodeIndex)
+		if (bIsFullCompile)
 		{
-			UEdGraphNode* Node = Context.LinearExecutionList[NodeIndex];
+			// Find the execution path (and make sure it has no cycles)
+			CreateExecutionSchedule(Context.SourceGraph->Nodes, Context.LinearExecutionList);
 
-			// Register nets in the schedule
-			if (FNodeHandlingFunctor* Handler = NodeHandlers.FindRef(Node->GetClass()))
+			// Create any user defined variables, this must occur before registering nets so that the properties are in place
+			UField** PropertyStorageLocation = &(Context.Function->Children);
+			CreateUserDefinedLocalVariablesForFunction(Context, PropertyStorageLocation);
+
+			for (int32 NodeIndex = 0; NodeIndex < Context.LinearExecutionList.Num(); ++NodeIndex)
 			{
-				Handler->RegisterNets(Context, Node);
+				UEdGraphNode* Node = Context.LinearExecutionList[NodeIndex];
+
+				// Register nets in the schedule
+				if (FNodeHandlingFunctor* Handler = NodeHandlers.FindRef(Node->GetClass()))
+				{
+					Handler->RegisterNets(Context, Node);
+				}
+				else
+				{
+					MessageLog.Error(*FString::Printf(*LOCTEXT("UnexpectedNodeType_Error", "Unexpected node type %s encountered at @@").ToString(), *(Node->GetClass()->GetName())), Node);
+				}
 			}
-			else
+		}
+		else
+		{
+			for (auto Node : Context.SourceGraph->Nodes)
 			{
-				MessageLog.Error(*FString::Printf(*LOCTEXT("UnexpectedNodeType_Error", "Unexpected node type %s encountered at @@").ToString(), *(Node->GetClass()->GetName())), Node);
+				const bool bGenerateParameters = Node->IsA<UK2Node_FunctionEntry>();
+				const bool bGenerateResult = Node->IsA<UK2Node_FunctionResult>();
+				if (bGenerateParameters || bGenerateResult)
+				{
+					if (FNodeHandlingFunctor* Handler = NodeHandlers.FindRef(Node->GetClass()))
+					{
+						Handler->RegisterNets(Context, Node);
+					}
+				}
 			}
 		}
 	
