Find dynamically bound nodes in function graphs as well as the uber graph

diff --git a/Engine/Source/Editor/KismetCompiler/Private/KismetCompiler.cpp b/Engine/Source/Editor/KismetCompiler/Private/KismetCompiler.cpp
--- a/Engine/Source/Editor/KismetCompiler/Private/KismetCompiler.cpp
+++ b/Engine/Source/Editor/KismetCompiler/Private/KismetCompiler.cpp
@@ -1733,23 +1733,26 @@ void FKismetCompilerContext::BuildDynamicBindingObjects(UBlueprintGeneratedClass
 {
 	Class->DynamicBindingObjects.Empty();
 
-	for (int NodeIndex = 0; NodeIndex < ConsolidatedEventGraph->Nodes.Num(); ++NodeIndex)
+	for (FKismetFunctionContext& FunctionContext : FunctionList)
 	{
-		UK2Node* Node = Cast<UK2Node>(ConsolidatedEventGraph->Nodes[NodeIndex]);
-
-		if (Node)
+		for (UEdGraphNode* GraphNode : FunctionContext.SourceGraph->Nodes)
 		{
-			UClass* DynamicBindingClass = Node->GetDynamicBindingClass();
+			UK2Node* Node = Cast<UK2Node>(GraphNode);
 
-			if (DynamicBindingClass)
+			if (Node)
 			{
-				UDynamicBlueprintBinding* DynamicBindingObject = Class->GetDynamicBindingObject(DynamicBindingClass);
-				if (DynamicBindingObject == NULL)
+				UClass* DynamicBindingClass = Node->GetDynamicBindingClass();
+
+				if (DynamicBindingClass)
 				{
-					DynamicBindingObject = ConstructObject<UDynamicBlueprintBinding>(DynamicBindingClass, Class);
-					Class->DynamicBindingObjects.Add(DynamicBindingObject);
+					UDynamicBlueprintBinding* DynamicBindingObject = Class->GetDynamicBindingObject(DynamicBindingClass);
+					if (DynamicBindingObject == NULL)
+					{
+						DynamicBindingObject = ConstructObject<UDynamicBlueprintBinding>(DynamicBindingClass, Class);
+						Class->DynamicBindingObjects.Add(DynamicBindingObject);
+					}
+					Node->RegisterDynamicBinding(DynamicBindingObject);
 				}
-				Node->RegisterDynamicBinding(DynamicBindingObject);
 			}
 		}
 	}
