BP: pins with default objects are properly handled by SpawnActor node

diff --git a/Engine/Source/Editor/KismetCompiler/Private/KismetCompilerMisc.cpp b/Engine/Source/Editor/KismetCompiler/Private/KismetCompilerMisc.cpp
--- a/Engine/Source/Editor/KismetCompiler/Private/KismetCompilerMisc.cpp
+++ b/Engine/Source/Editor/KismetCompiler/Private/KismetCompilerMisc.cpp
@@ -517,8 +517,9 @@ UEdGraphPin* FKismetCompilerUtilities::GenerateAssignmentNodes(class FKismetComp
 	{
 		// Only create 'set param by name' node if this pin is linked to something
 		UEdGraphPin* OrgPin = SpawnNode->Pins[PinIdx];
+		const bool bHasDefaultValue = !OrgPin->DefaultValue.IsEmpty() || !OrgPin->DefaultTextValue.IsEmpty() || OrgPin->DefaultObject;
 		if (NULL == CallBeginSpawnNode->FindPin(OrgPin->PinName) &&
-			(OrgPin->LinkedTo.Num() > 0 || OrgPin->DefaultValue != FString()))
+			(OrgPin->LinkedTo.Num() > 0 || bHasDefaultValue))
 		{
 			if( OrgPin->LinkedTo.Num() == 0 )
 			{
