[334276] fixed - compiling a BP containing nodes with external dependencies unnecessarily causes other BPs to be marked as needing to be saved, even though no changes were made

- clear the package dirty state on external BP dependencies if they did not have any unsaved changes prior to refreshing. since the external BPs will be recompiled/refreshed on next load, there's not really a need to require them to be saved out. they do however need to be recompiled; that part remains unchanged and that step still happens during reinstancing.

diff --git a/Engine/Source/Editor/KismetCompiler/Private/KismetCompiler.cpp b/Engine/Source/Editor/KismetCompiler/Private/KismetCompiler.cpp
--- a/Engine/Source/Editor/KismetCompiler/Private/KismetCompiler.cpp
+++ b/Engine/Source/Editor/KismetCompiler/Private/KismetCompiler.cpp
@@ -3439,9 +3439,19 @@ void FKismetCompilerContext::Compile()
 		FBlueprintEditorUtils::GetDependentBlueprints(Blueprint, DependentBlueprints);
 		for (auto CurrentBP : DependentBlueprints)
 		{
+			// Get the current dirty state of the package
+			UPackage* const Package = Cast<UPackage>(CurrentBP->GetOutermost());
+			const bool bStartedWithUnsavedChanges = Package != nullptr ? Package->IsDirty() : true;
+
 			CurrentBP->Status = BS_Dirty;
 			FBlueprintEditorUtils::RefreshExternalBlueprintDependencyNodes(CurrentBP, NewClass);
 			CurrentBP->BroadcastChanged();
+
+			// Clear the package dirty state if it did not initially have any unsaved changes to begin with
+			if(Package != nullptr && Package->IsDirty() && !bStartedWithUnsavedChanges)
+			{
+				Package->SetDirtyFlag(false);
+			}
 		}
 	}
 
