ttp341911 BLUEPRINTS: Array output being cached and not cleared, supplying inacurate data.

diff --git a/Engine/Source/Editor/KismetCompiler/Private/KismetCompilerVMBackend.cpp b/Engine/Source/Editor/KismetCompiler/Private/KismetCompilerVMBackend.cpp
--- a/Engine/Source/Editor/KismetCompiler/Private/KismetCompilerVMBackend.cpp
+++ b/Engine/Source/Editor/KismetCompiler/Private/KismetCompilerVMBackend.cpp
@@ -612,6 +612,30 @@ public:
 	{
 		UFunction* FunctionToCall = Statement.FunctionToCall;
 		check(FunctionToCall);
+
+		if (FunctionToCall->HasAllFunctionFlags(FUNC_Native))
+		{
+			// Array output parameters are cleared, in case the native function doesn't clear them before filling.
+			int32 NumParams = 0;
+			for (TFieldIterator<UProperty> PropIt(FunctionToCall); PropIt && (PropIt->PropertyFlags & CPF_Parm); ++PropIt)
+			{
+				UProperty* Param = *PropIt;
+				check(Param);
+				const bool bShouldParameterBeCleared = Param->IsA<UArrayProperty>()
+					&& Param->HasAllPropertyFlags(CPF_Parm | CPF_OutParm)
+					&& !Param->HasAnyPropertyFlags(CPF_ReferenceParm | CPF_ConstParm | CPF_ReturnParm);
+				if (bShouldParameterBeCleared)
+				{
+					// SetArray instruction will be called with empty parameter list.
+					Writer << EX_SetArray;
+					FBPTerminal* ArrayTerm = Statement.RHS[NumParams];
+					ensure(ArrayTerm && !ArrayTerm->bIsLiteral && !Statement.ArrayCoersionTermMap.Find(ArrayTerm));
+					EmitTerm(ArrayTerm, Param);
+					Writer << EX_EndArray;
+				}
+				NumParams += Param->HasAnyPropertyFlags(CPF_ReturnParm) ? 0 : 1;
+			}
+		}
 		
 		// The target label will only ever be set on a call function when calling into the Ubergraph, which requires a patchup
 		// or when re-entering from a latent function which requires a different kind of patchup
