This changelist adds some timeline validation code to the compiler to strip disconnected timelines from actor instances and improve the cost of checking to tick them.
#TTP 342661 - UDN: LIVE: BP: Inactive/disconnected Timelines still take resources in editor
#Branch UE4
#Proj KismetCompiler

#Change Added a ValidatedTimelines list to the compiler context, when we discover timelines are connected when interrogating the graphs they are placed in this list. The list is then used to filter which timelines make it into the Generated Class.

diff --git a/Engine/Source/Editor/KismetCompiler/Private/KismetCompiler.cpp b/Engine/Source/Editor/KismetCompiler/Private/KismetCompiler.cpp
--- a/Engine/Source/Editor/KismetCompiler/Private/KismetCompiler.cpp
+++ b/Engine/Source/Editor/KismetCompiler/Private/KismetCompiler.cpp
@@ -501,7 +501,7 @@ void FKismetCompilerContext::CreateClassVariablesFromBlueprint()
 	{
 		UTimelineTemplate* Timeline = Blueprint->Timelines[TimelineIndex];
 		// Not fatal if NULL, but shouldn't happen
-		if (!Timeline)
+		if (!Timeline || !ValidatedTimelines.Find( Timeline->GetFName() ))
 		{
 			continue;
 		}
@@ -1718,7 +1718,24 @@ void FKismetCompilerContext::FinishCompilingClass(UClass* Class)
 		BPGClass->SimpleConstructionScript = NULL;
 
 		BPGClass->ComponentTemplates = Blueprint->ComponentTemplates;
-		BPGClass->Timelines = Blueprint->Timelines;
+
+		if( bIsFullCompile )
+		{
+			// Add only the timelines we have validated as having connections in the graph
+			for( int32 TimelineIndex = 0; TimelineIndex < Blueprint->Timelines.Num(); ++TimelineIndex )
+			{
+				UTimelineTemplate* TimeLineInstance = Blueprint->Timelines[ TimelineIndex ];
+
+				if( ValidatedTimelines.Find( TimeLineInstance->GetFName() ) != INDEX_NONE )
+				{
+					BPGClass->Timelines.Add( TimeLineInstance );
+				}
+			}
+		}
+		else
+		{
+			BPGClass->Timelines = Blueprint->Timelines;
+		}
 		BPGClass->SimpleConstructionScript = Blueprint->SimpleConstructionScript;
 	}
 
@@ -1840,6 +1857,9 @@ void FKismetCompilerContext::ExpandTimelineNodes(UEdGraph* SourceGraph)
 				// Only create nodes for play/stop if they are actually connected - otherwise we get a 'unused node being pruned' warning
 				if(bPlayPinConnected || bPlayFromStartPinConnected || bStopPinConnected || bReversePinConnected || bReverseFromEndPinConnected || bSetNewTimePinConnected)
 				{
+					// This timeline is wired up so it is now considered validated
+					ValidatedTimelines.Add( TimelineNode->TimelineName );
+
 					// First create 'get var' node to get the timeline object
 					UK2Node_VariableGet* GetTimelineNode = SpawnIntermediateNode<UK2Node_VariableGet>(TimelineNode, SourceGraph);
 					GetTimelineNode->VariableReference.SetSelfMember(TimelineNode->TimelineName);
@@ -1905,23 +1925,22 @@ void FKismetCompilerContext::ExpandTimelineNodes(UEdGraph* SourceGraph)
 							MovePinLinksToIntermediate(*NewTimePin, *InputPin);
 						}
 					}
-				}
+					// Create event to call on each update
+					UFunction* EventSigFunc = UTimelineComponent::GetTimelineEventSignature();
 
-				// Create event to call on each update
-				UFunction* EventSigFunc = UTimelineComponent::GetTimelineEventSignature();
-
-				// Create event nodes for any event tracks
-				for(int32 EventTrackIdx=0; EventTrackIdx<Timeline->EventTracks.Num(); EventTrackIdx++)
-				{
-					FName EventTrackName =  Timeline->EventTracks[EventTrackIdx].TrackName;
-					CreatePinEventNodeForTimelineFunction(TimelineNode, SourceGraph, Timeline->GetEventTrackFunctionName(EventTrackIdx), EventTrackName.ToString(), EventSigFunc->GetFName());
-				}
+					// Create event nodes for any event tracks
+					for(int32 EventTrackIdx=0; EventTrackIdx<Timeline->EventTracks.Num(); EventTrackIdx++)
+					{
+						FName EventTrackName =  Timeline->EventTracks[EventTrackIdx].TrackName;
+						CreatePinEventNodeForTimelineFunction(TimelineNode, SourceGraph, Timeline->GetEventTrackFunctionName(EventTrackIdx), EventTrackName.ToString(), EventSigFunc->GetFName());
+					}
 
-				// Generate Update Pin Event Node
-				CreatePinEventNodeForTimelineFunction(TimelineNode, SourceGraph, Timeline->GetUpdateFunctionName(), TEXT("Update"), EventSigFunc->GetFName());
+					// Generate Update Pin Event Node
+					CreatePinEventNodeForTimelineFunction(TimelineNode, SourceGraph, Timeline->GetUpdateFunctionName(), TEXT("Update"), EventSigFunc->GetFName());
 
-				// Generate Finished Pin Event Node
-				CreatePinEventNodeForTimelineFunction(TimelineNode, SourceGraph, Timeline->GetFinishedFunctionName(), TEXT("Finished"), EventSigFunc->GetFName());
+					// Generate Finished Pin Event Node
+					CreatePinEventNodeForTimelineFunction(TimelineNode, SourceGraph, Timeline->GetFinishedFunctionName(), TEXT("Finished"), EventSigFunc->GetFName());
+				}
 			}			
 		}
 	}
diff --git a/Engine/Source/Editor/KismetCompiler/Public/KismetCompiler.h b/Engine/Source/Editor/KismetCompiler/Public/KismetCompiler.h
--- a/Engine/Source/Editor/KismetCompiler/Public/KismetCompiler.h
+++ b/Engine/Source/Editor/KismetCompiler/Public/KismetCompiler.h
@@ -35,6 +35,10 @@ protected:
 
 	// List of functions currently allocated
 	TIndirectArray<FKismetFunctionContext> FunctionList;
+
+	// List of validated Timelines
+	TArray<FName> ValidatedTimelines;
+
 protected:
 	// This struct holds the various compilation options, such as which passes to perform, whether to save intermediate results, etc
 	FKismetCompilerOptions CompileOptions;
