ttp343023 BP: incompatible overriden fuinctions wont cause a failed ensure.

diff --git a/Engine/Source/Editor/KismetCompiler/Private/KismetCompiler.cpp b/Engine/Source/Editor/KismetCompiler/Private/KismetCompiler.cpp
--- a/Engine/Source/Editor/KismetCompiler/Private/KismetCompiler.cpp
+++ b/Engine/Source/Editor/KismetCompiler/Private/KismetCompiler.cpp
@@ -1482,12 +1482,16 @@ void FKismetCompilerContext::FinishCompilingFunction(FKismetFunctionContext& Con
 	{
 		Function->FunctionFlags |= (OverridenFunction->FunctionFlags & (FUNC_FuncInherit | FUNC_Public | FUNC_Protected | FUNC_Private));
 
-		if( (Function->FunctionFlags & FUNC_AccessSpecifiers) != (OverridenFunction->FunctionFlags & FUNC_AccessSpecifiers) )
+		if ((Function->FunctionFlags & FUNC_AccessSpecifiers) != (OverridenFunction->FunctionFlags & FUNC_AccessSpecifiers))
 		{
 			MessageLog.Error(*LOCTEXT("IncompatibleAccessSpecifier_Error", "Access specifier is not compatible the parent function @@").ToString(), Context.EntryPoint);
 		}
 
-		ensure((Function->FunctionFlags & FUNC_FuncOverrideMatch) == (OverridenFunction->FunctionFlags & FUNC_FuncOverrideMatch));
+		const uint32 OverrideFlagsToCheck = (FUNC_FuncOverrideMatch & ~FUNC_AccessSpecifiers);
+		if ((Function->FunctionFlags & OverrideFlagsToCheck) != (OverridenFunction->FunctionFlags & OverrideFlagsToCheck))
+		{
+			MessageLog.Error(*LOCTEXT("IncompatibleOverrideFlags_Error", "Overriden function is not compatible with the parent function @@. Check flags: Exec, Final, Static.").ToString(), Context.EntryPoint);
+		}
 
 		// Copy metadata from parent function as well
 		UMetaData::CopyMetadata(OverridenFunction, Function);
