Reworking of changelist 2290702 to avoid issues with timeline pins dissapearing. This modification flags TimelineTemplates as wired during BP compilation and checks the results during timeline component instantiation pruning disconnected timeline components.
#TTP 342661 - UDN: LIVE: BP: Inactive/disconnected Timelines still take resources in editor
#Branch UE4
#Proj KismetCompiler

#Change Modified the way the original fix worked due to issues with curve pin's dissapearing in graph ( no data was destroyed ). This fix leaves all the timeline templates in tact in the CDO but flags them as being wired and prunes unwired timelines from becoming actor components.

diff --git a/Engine/Source/Editor/KismetCompiler/Private/KismetCompiler.cpp b/Engine/Source/Editor/KismetCompiler/Private/KismetCompiler.cpp
--- a/Engine/Source/Editor/KismetCompiler/Private/KismetCompiler.cpp
+++ b/Engine/Source/Editor/KismetCompiler/Private/KismetCompiler.cpp
@@ -1841,8 +1841,11 @@ void FKismetCompilerContext::ExpandTimelineNodes(UEdGraph* SourceGraph)
 				UEdGraphPin* SetTimePin = TimelineNode->GetSetNewTimePin();
 				bool bSetNewTimePinConnected = (SetTimePin->LinkedTo.Num() > 0);
 
+				// Set the timeline template as wired/not wired for component pruning later
+				Timeline->bValidatedAsWired = bPlayPinConnected || bPlayFromStartPinConnected || bStopPinConnected || bReversePinConnected || bReverseFromEndPinConnected || bSetNewTimePinConnected;
+
 				// Only create nodes for play/stop if they are actually connected - otherwise we get a 'unused node being pruned' warning
-				if(bPlayPinConnected || bPlayFromStartPinConnected || bStopPinConnected || bReversePinConnected || bReverseFromEndPinConnected || bSetNewTimePinConnected)
+				if(Timeline->bValidatedAsWired)
 				{
 					// First create 'get var' node to get the timeline object
 					UK2Node_VariableGet* GetTimelineNode = SpawnIntermediateNode<UK2Node_VariableGet>(TimelineNode, SourceGraph);
diff --git a/Engine/Source/Runtime/Engine/Classes/Engine/TimelineTemplate.h b/Engine/Source/Runtime/Engine/Classes/Engine/TimelineTemplate.h
--- a/Engine/Source/Runtime/Engine/Classes/Engine/TimelineTemplate.h
+++ b/Engine/Source/Runtime/Engine/Classes/Engine/TimelineTemplate.h
@@ -131,6 +131,10 @@ class UTimelineTemplate : public UObject
 	UPROPERTY(EditAnywhere, BlueprintReadWrite, Category=TimelineTemplate)
 	uint32 bReplicated:1;
 
+	/** Compiler Validated As Wired up */
+	UPROPERTY()
+	uint32 bValidatedAsWired:1;
+
 	/** Set of event tracks */
 	UPROPERTY()
 	TArray<struct FTTEventTrack> EventTracks;
diff --git a/Engine/Source/Runtime/Engine/Private/BlueprintGeneratedClass.cpp b/Engine/Source/Runtime/Engine/Private/BlueprintGeneratedClass.cpp
--- a/Engine/Source/Runtime/Engine/Private/BlueprintGeneratedClass.cpp
+++ b/Engine/Source/Runtime/Engine/Private/BlueprintGeneratedClass.cpp
@@ -293,8 +293,8 @@ void UBlueprintGeneratedClass::CreateComponentsForActor(AActor* Actor) const
 	{
 		const UTimelineTemplate* TimelineTemplate = Timelines[i];
 
-		// Not fatal if NULL, but shouldn't happen
-		if(!TimelineTemplate)
+		// Not fatal if NULL, but shouldn't happen and ignored if not wired up in graph
+		if(!TimelineTemplate||!TimelineTemplate->bValidatedAsWired)
 		{
 			continue;
 		}
