KismetCppBackend: EmitMetaCastStatement added, logged in one chunk.

diff --git a/Engine/Source/Editor/KismetCompiler/Private/KismetCompiler.cpp b/Engine/Source/Editor/KismetCompiler/Private/KismetCompiler.cpp
--- a/Engine/Source/Editor/KismetCompiler/Private/KismetCompiler.cpp
+++ b/Engine/Source/Editor/KismetCompiler/Private/KismetCompiler.cpp
@@ -3377,16 +3377,8 @@ void FKismetCompilerContext::Compile()
 			// The C++ backend is currently only for debugging, so it's only run if the output will be visible
 			Backend_CPP.GenerateCodeFromClass(NewClass, FunctionList, !bIsFullCompile);
 		
-			// need to break it down per line to prevent the log from failing to emit it
-			TArray<FString> Lines;
-			FString TotalString = FString::Printf(TEXT("\n\n\n[header]\n\n\n%s[body]\n%s"), *Backend_CPP.Header, *Backend_CPP.Body);
-
-			TotalString.ParseIntoArray(&Lines, TEXT("\n"), true);
-			for (int32 I=0; I<Lines.Num(); ++I)
-			{
-				FString Line = Lines[I];
-				UE_LOG(LogK2Compiler, Log, TEXT("%s"), *Line);
-			}
+			UE_LOG(LogK2Compiler, Log, TEXT("[header]\n\n\n%s"), *Backend_CPP.Header);
+			UE_LOG(LogK2Compiler, Log, TEXT("[body]\n\n\n%s"), *Backend_CPP.Body);
 		}
 
 		// Always run the VM backend, it's needed for more than just debug printing
diff --git a/Engine/Source/Editor/KismetCompiler/Private/KismetCompilerBackend.h b/Engine/Source/Editor/KismetCompiler/Private/KismetCompilerBackend.h
--- a/Engine/Source/Editor/KismetCompiler/Private/KismetCompilerBackend.h
+++ b/Engine/Source/Editor/KismetCompiler/Private/KismetCompilerBackend.h
@@ -108,6 +108,7 @@ public:
 	void EmitEndOfThreadStatement(FKismetFunctionContext& FunctionContext, const FString& ReturnValueString);
 	void EmitReturnStatement(FKismetFunctionContext& FunctionContext, const FString& ReturnValueString);
 	void EmitRemoveMulticastDelegateStatement(FKismetFunctionContext& FunctionContext, FBlueprintCompiledStatement& Statement);
+	void EmitMetaCastStatement(FKismetFunctionContext& FunctionContext, FBlueprintCompiledStatement& Statement);
 	
 	/** Emits local variable declarations for a function */
 	void DeclareLocalVariables(FKismetFunctionContext& FunctionContext, TArray<UProperty*>& LocalVariables);
diff --git a/Engine/Source/Editor/KismetCompiler/Private/KismetCompilerCppBackend.cpp b/Engine/Source/Editor/KismetCompiler/Private/KismetCompilerCppBackend.cpp
--- a/Engine/Source/Editor/KismetCompiler/Private/KismetCompilerCppBackend.cpp
+++ b/Engine/Source/Editor/KismetCompiler/Private/KismetCompilerCppBackend.cpp
@@ -472,10 +472,20 @@ void FKismetCppBackend::EmitDynamicCastStatement(FKismetFunctionContext& Functio
 	FString ObjectValue = TermToText(Statement.RHS[1], (UProperty*)(GetDefault<UObjectProperty>()));
 	FString CastedValue = TermToText(Statement.LHS, (UProperty*)(GetDefault<UObjectProperty>()));
 
-	Emit(Body, *FString::Printf(TEXT("\t\t\t %s = Cast<%s>(%s);\n"),
+	Emit(Body, *FString::Printf(TEXT("\t\t\t%s = Cast<%s>(%s);\n"),
 		*CastedValue, *TargetClass, *ObjectValue));
 }
 
+void FKismetCppBackend::EmitMetaCastStatement(FKismetFunctionContext& FunctionContext, FBlueprintCompiledStatement& Statement)
+{
+	FString DesiredClass	= TermToText(Statement.RHS[0], (UProperty*)(GetDefault<UClassProperty>()));
+	FString SourceClass		= TermToText(Statement.RHS[1], (UProperty*)(GetDefault<UClassProperty>()));
+	FString Destination		= TermToText(Statement.LHS, (UProperty*)(GetDefault<UClassProperty>()));
+
+	Emit(Body, *FString::Printf(TEXT("\t\t\t%s = ((%s)->IsChildOf(DesiredClass)) ? SourceClass : NULL;\n"), 
+		*Destination, *SourceClass, *DesiredClass, *SourceClass));
+}
+
 void FKismetCppBackend::EmitObjectToBoolStatement(FKismetFunctionContext& FunctionContext, FBlueprintCompiledStatement& Statement)
 {
 	FString ObjectTarget = TermToText(Statement.RHS[0]);
@@ -837,6 +847,12 @@ void FKismetCppBackend::ConstructFunction(FKismetFunctionContext& FunctionContex
 					case KCST_Comment:
 						Emit(Body, *FString::Printf(TEXT("\t\t\t// %s\n"), *Statement.Comment));
 						break;
+					case KCST_MetaCast:
+						EmitMetaCastStatement(FunctionContext, Statement);
+						break;
+					case KCST_Return:
+						Emit(Body, TEXT("\t\t\t// Return statement.\n"));
+						break;
 					default:
 						Emit(Body, TEXT("\t// Warning: Ignoring unsupported statement\n"));
 						UE_LOG(LogK2Compiler, Warning, TEXT("C++ backend encountered unsupported statement type %d"), (int32)Statement.Type);
