#ue4 When class properties are missing their target class, replace with UObject and throw a warning, so they can be fixed

diff --git a/Engine/Source/Editor/KismetCompiler/Private/KismetCompilerMisc.cpp b/Engine/Source/Editor/KismetCompiler/Private/KismetCompilerMisc.cpp
--- a/Engine/Source/Editor/KismetCompiler/Private/KismetCompilerMisc.cpp
+++ b/Engine/Source/Editor/KismetCompiler/Private/KismetCompilerMisc.cpp
@@ -710,6 +710,19 @@ UProperty* FKismetCompilerUtilities::CreatePropertyOnScope(UStruct* Scope, const
 	else if (Type.PinCategory == Schema->PC_Class)
 	{
 		UClass* SubType = Cast<UClass>(Type.PinSubCategoryObject.Get());
+		
+		if (SubType == NULL)
+		{
+			// If this is from a degenerate pin, because the object type has been removed, default this to a UObject subtype so we can make a dummy term for it to allow the compiler to continue
+			SubType = UObject::StaticClass();
+
+			MessageLog.Warning(
+				*FString::Printf(
+				*LOCTEXT("InvalidClassForField_Error", "Invalid property '%s' class, replaced with Object.  Please fix or remove.").ToString(),
+				*PropertyName.ToString()
+				));
+		}
+
 		if (SubType != NULL)
 		{
 			UClassProperty* NewPropertyClass = NewNamedObject<UClassProperty>(PropertyScope, ValidatedPropertyName, ObjectFlags);
