Fix for timeline components being removed incorrectly, this should fix the issues seen in tappy chicken.
#TTP 346668 - Blueprint nodes being possibly optimized out
#Branch UE4
#Proj  KismetCompiler

#change fixed a some incomplete logic in the timelines that didn't account for autostart or variable get usage.

diff --git a/Engine/Source/Editor/KismetCompiler/Private/KismetCompiler.cpp b/Engine/Source/Editor/KismetCompiler/Private/KismetCompiler.cpp
--- a/Engine/Source/Editor/KismetCompiler/Private/KismetCompiler.cpp
+++ b/Engine/Source/Editor/KismetCompiler/Private/KismetCompiler.cpp
@@ -1822,7 +1822,7 @@ void FKismetCompilerContext::ExpandTimelineNodes(UEdGraph* SourceGraph)
 			UTimelineTemplate* Timeline = Blueprint->FindTimelineTemplateByVariableName(TimelineNode->TimelineName);
 			if (Timeline != NULL)
 			{
-				FString TimelineNameString = TimelineNode->TimelineName.ToString();
+				const FString TimelineNameString = TimelineNode->TimelineName.ToString();
 
 				UEdGraphPin* PlayPin = TimelineNode->GetPlayPin();
 				bool bPlayPinConnected = (PlayPin->LinkedTo.Num() > 0);
@@ -1842,11 +1842,21 @@ void FKismetCompilerContext::ExpandTimelineNodes(UEdGraph* SourceGraph)
 				UEdGraphPin* SetTimePin = TimelineNode->GetSetNewTimePin();
 				bool bSetNewTimePinConnected = (SetTimePin->LinkedTo.Num() > 0);
 
+				UEdGraphPin* UpdatePin = TimelineNode->GetUpdatePin();
+				bool bUpdatePinConnected = (UpdatePin->LinkedTo.Num() > 0);
+
+				UEdGraphPin* FinishedPin = TimelineNode->GetFinishedPin();
+				bool bFinishedPinConnected = (FinishedPin->LinkedTo.Num() > 0);
+				
+
 				// Set the timeline template as wired/not wired for component pruning later
-				Timeline->bValidatedAsWired = bPlayPinConnected || bPlayFromStartPinConnected || bStopPinConnected || bReversePinConnected || bReverseFromEndPinConnected || bSetNewTimePinConnected;
+				const bool bWiredIn = bPlayPinConnected || bPlayFromStartPinConnected || bStopPinConnected || bReversePinConnected || bReverseFromEndPinConnected || bSetNewTimePinConnected;
+				const bool bWiredOut = bUpdatePinConnected || bFinishedPinConnected;
+
+				Timeline->bValidatedAsWired = bWiredIn || ( Timeline->bAutoPlay && bWiredOut );
 
 				// Only create nodes for play/stop if they are actually connected - otherwise we get a 'unused node being pruned' warning
-				if(Timeline->bValidatedAsWired)
+				if(bWiredIn)
 				{
 					// First create 'get var' node to get the timeline object
 					UK2Node_VariableGet* GetTimelineNode = SpawnIntermediateNode<UK2Node_VariableGet>(TimelineNode, SourceGraph);
@@ -1932,6 +1942,23 @@ void FKismetCompilerContext::ExpandTimelineNodes(UEdGraph* SourceGraph)
 				CreatePinEventNodeForTimelineFunction(TimelineNode, SourceGraph, Timeline->GetFinishedFunctionName(), TEXT("Finished"), EventSigFunc->GetFName());
 			}			
 		}
+		else if( UK2Node_VariableGet* VarNode = Cast<UK2Node_VariableGet>( SourceGraph->Nodes[ChildIndex] ))
+		{
+			// Check for Timeline Variable Get Nodes
+			UEdGraphPin* const ValuePin = VarNode->GetValuePin();
+
+			if( ValuePin && ValuePin->LinkedTo.Num() > 0  )
+			{
+				UClass* ValueClass = ValuePin->PinType.PinSubCategoryObject.IsValid() ? Cast<UClass>( ValuePin->PinType.PinSubCategoryObject.Get() ) : nullptr;
+				if( ValueClass == UTimelineComponent::StaticClass() )
+				{
+					if( UTimelineTemplate* TimelineTemplate = Blueprint->FindTimelineTemplateByVariableName( FName( *ValuePin->PinName ) ))
+					{
+						TimelineTemplate->bValidatedAsWired = true;
+					}
+				}
+			}
+		}
 	}
 }
 
