Additional fix required for missing timelines, had to separate out the checking for external play nodes on timelines.
#TTP 346668 - Blueprint nodes being possibly optimized out
#Branch UE4
#Proj  KismetCompiler

#change fixed an ordering issue when checking timeline validity and managed to get tappy chicken to work this time to verify the fix against.

diff --git a/Engine/Source/Editor/KismetCompiler/Private/KismetCompiler.cpp b/Engine/Source/Editor/KismetCompiler/Private/KismetCompiler.cpp
--- a/Engine/Source/Editor/KismetCompiler/Private/KismetCompiler.cpp
+++ b/Engine/Source/Editor/KismetCompiler/Private/KismetCompiler.cpp
@@ -1814,151 +1814,176 @@ UK2Node_CallFunction* FKismetCompilerContext::CreateCallTimelineFunction(UK2Node
 /** Expand timeline nodes into necessary nodes */
 void FKismetCompilerContext::ExpandTimelineNodes(UEdGraph* SourceGraph)
 {
+	// Timeline Pair helper
+	struct FTimelinePair
+	{
+	public:
+		FTimelinePair( UK2Node_Timeline* InNode, UTimelineTemplate* InTemplate )
+		: Node( InNode )
+		, Template( InTemplate )
+		{
+		}
+	public:
+		UK2Node_Timeline* const Node;
+		UTimelineTemplate* Template;
+	};
+
+	TArray<FName> TimelinePlayNodes;
+	TArray<FTimelinePair> Timelines;
+	// Extract timeline pairings and external play nodes
 	for (int32 ChildIndex = 0; ChildIndex < SourceGraph->Nodes.Num(); ++ChildIndex)
 	{
-		UK2Node_Timeline* TimelineNode = Cast<UK2Node_Timeline>( SourceGraph->Nodes[ChildIndex] );
-		if (TimelineNode != NULL)
+		if (UK2Node_Timeline* TimelineNode = Cast<UK2Node_Timeline>( SourceGraph->Nodes[ChildIndex]))
 		{
 			UTimelineTemplate* Timeline = Blueprint->FindTimelineTemplateByVariableName(TimelineNode->TimelineName);
 			if (Timeline != NULL)
 			{
-				const FString TimelineNameString = TimelineNode->TimelineName.ToString();
+				Timelines.Add(FTimelinePair( TimelineNode, Timeline ));
+			}
+		}
+		else if( UK2Node_VariableGet* VarNode = Cast<UK2Node_VariableGet>( SourceGraph->Nodes[ChildIndex] ))
+		{
+			// Check for Timeline Variable Get Nodes
+			UEdGraphPin* const ValuePin = VarNode->GetValuePin();
 
-				UEdGraphPin* PlayPin = TimelineNode->GetPlayPin();
-				bool bPlayPinConnected = (PlayPin->LinkedTo.Num() > 0);
+			if( ValuePin && ValuePin->LinkedTo.Num() > 0  )
+			{
+				UClass* ValueClass = ValuePin->PinType.PinSubCategoryObject.IsValid() ? Cast<UClass>( ValuePin->PinType.PinSubCategoryObject.Get() ) : nullptr;
+				if( ValueClass == UTimelineComponent::StaticClass() )
+				{
+					const FName PinName( *ValuePin->PinName );
+					if( UTimelineTemplate* TimelineTemplate = Blueprint->FindTimelineTemplateByVariableName( PinName ))
+					{
+						TimelinePlayNodes.Add( PinName );
+					}
+				}
+			}
+		}
+	}
+	// Expand and validate timelines
+	for ( auto TimelinePair : Timelines )
+	{
+		UK2Node_Timeline* const TimelineNode = TimelinePair.Node;
+		UTimelineTemplate* Timeline = TimelinePair.Template;
+		const FString TimelineNameString = TimelineNode->TimelineName.ToString();
 
-				UEdGraphPin* PlayFromStartPin = TimelineNode->GetPlayFromStartPin();
-				bool bPlayFromStartPinConnected = (PlayFromStartPin->LinkedTo.Num() > 0);
+		UEdGraphPin* PlayPin = TimelineNode->GetPlayPin();
+		bool bPlayPinConnected = (PlayPin->LinkedTo.Num() > 0);
 
-				UEdGraphPin* StopPin = TimelineNode->GetStopPin();
-				bool bStopPinConnected = (StopPin->LinkedTo.Num() > 0);
+		UEdGraphPin* PlayFromStartPin = TimelineNode->GetPlayFromStartPin();
+		bool bPlayFromStartPinConnected = (PlayFromStartPin->LinkedTo.Num() > 0);
 
-				UEdGraphPin* ReversePin = TimelineNode->GetReversePin();
-				bool bReversePinConnected = (ReversePin->LinkedTo.Num() > 0);
+		UEdGraphPin* StopPin = TimelineNode->GetStopPin();
+		bool bStopPinConnected = (StopPin->LinkedTo.Num() > 0);
 
-				UEdGraphPin* ReverseFromEndPin = TimelineNode->GetReverseFromEndPin();
-				bool bReverseFromEndPinConnected = (ReverseFromEndPin->LinkedTo.Num() > 0);
+		UEdGraphPin* ReversePin = TimelineNode->GetReversePin();
+		bool bReversePinConnected = (ReversePin->LinkedTo.Num() > 0);
 
-				UEdGraphPin* SetTimePin = TimelineNode->GetSetNewTimePin();
-				bool bSetNewTimePinConnected = (SetTimePin->LinkedTo.Num() > 0);
+		UEdGraphPin* ReverseFromEndPin = TimelineNode->GetReverseFromEndPin();
+		bool bReverseFromEndPinConnected = (ReverseFromEndPin->LinkedTo.Num() > 0);
 
-				UEdGraphPin* UpdatePin = TimelineNode->GetUpdatePin();
-				bool bUpdatePinConnected = (UpdatePin->LinkedTo.Num() > 0);
+		UEdGraphPin* SetTimePin = TimelineNode->GetSetNewTimePin();
+		bool bSetNewTimePinConnected = (SetTimePin->LinkedTo.Num() > 0);
 
-				UEdGraphPin* FinishedPin = TimelineNode->GetFinishedPin();
-				bool bFinishedPinConnected = (FinishedPin->LinkedTo.Num() > 0);
-				
+		UEdGraphPin* UpdatePin = TimelineNode->GetUpdatePin();
+		bool bUpdatePinConnected = (UpdatePin->LinkedTo.Num() > 0);
 
-				// Set the timeline template as wired/not wired for component pruning later
-				const bool bWiredIn = bPlayPinConnected || bPlayFromStartPinConnected || bStopPinConnected || bReversePinConnected || bReverseFromEndPinConnected || bSetNewTimePinConnected;
-				const bool bWiredOut = bUpdatePinConnected || bFinishedPinConnected;
+		UEdGraphPin* FinishedPin = TimelineNode->GetFinishedPin();
+		bool bFinishedPinConnected = (FinishedPin->LinkedTo.Num() > 0);
+		
 
-				Timeline->bValidatedAsWired = bWiredIn || ( Timeline->bAutoPlay && bWiredOut );
+		// Set the timeline template as wired/not wired for component pruning later
+		const bool bWiredIn = bPlayPinConnected || bPlayFromStartPinConnected || bStopPinConnected || bReversePinConnected || bReverseFromEndPinConnected || bSetNewTimePinConnected;
+		const bool bWiredOut = bUpdatePinConnected || bFinishedPinConnected;
+		const bool bPlayWired = Timeline->bAutoPlay || ( TimelinePlayNodes.Find( TimelineNode->TimelineName ) != INDEX_NONE );
 
-				// Only create nodes for play/stop if they are actually connected - otherwise we get a 'unused node being pruned' warning
-				if(bWiredIn)
-				{
-					// First create 'get var' node to get the timeline object
-					UK2Node_VariableGet* GetTimelineNode = SpawnIntermediateNode<UK2Node_VariableGet>(TimelineNode, SourceGraph);
-					GetTimelineNode->VariableReference.SetSelfMember(TimelineNode->TimelineName);
-					GetTimelineNode->AllocateDefaultPins();
-
-					// Debug data: Associate the timeline node instance with the property that was created earlier
-					UProperty* AssociatedTimelineInstanceProperty = TimelineToMemberVariableMap.FindChecked(Timeline);
-					if (AssociatedTimelineInstanceProperty != NULL)
-					{
-						UObject* TrueSourceObject = MessageLog.FindSourceObject(TimelineNode);
-						NewClass->GetDebugData().RegisterClassPropertyAssociation(TrueSourceObject, AssociatedTimelineInstanceProperty);
-					}
+		Timeline->bValidatedAsWired = bWiredIn || ( bPlayWired && bWiredOut );
 
-					// Get the variable output pin
-					UEdGraphPin* TimelineVarPin = GetTimelineNode->FindPin(TimelineNameString);
+		// Only create nodes for play/stop if they are actually connected - otherwise we get a 'unused node being pruned' warning
+		if(bWiredIn)
+		{
+			// First create 'get var' node to get the timeline object
+			UK2Node_VariableGet* GetTimelineNode = SpawnIntermediateNode<UK2Node_VariableGet>(TimelineNode, SourceGraph);
+			GetTimelineNode->VariableReference.SetSelfMember(TimelineNode->TimelineName);
+			GetTimelineNode->AllocateDefaultPins();
 
-					// This might fail if this is the first compile after adding the timeline (property doesn't exist yet) - in that case, manually add the output pin
-					if(TimelineVarPin == NULL)
-					{
-						TimelineVarPin = GetTimelineNode->CreatePin(EGPD_Output, Schema->PC_Object, TEXT(""), UTimelineComponent::StaticClass(), false, false, TimelineNode->TimelineName.ToString());
-					}
+			// Debug data: Associate the timeline node instance with the property that was created earlier
+			UProperty* AssociatedTimelineInstanceProperty = TimelineToMemberVariableMap.FindChecked(Timeline);
+			if (AssociatedTimelineInstanceProperty != NULL)
+			{
+				UObject* TrueSourceObject = MessageLog.FindSourceObject(TimelineNode);
+				NewClass->GetDebugData().RegisterClassPropertyAssociation(TrueSourceObject, AssociatedTimelineInstanceProperty);
+			}
 
-					if(bPlayPinConnected)
-					{
-						static FName PlayName(GET_FUNCTION_NAME_CHECKED(UTimelineComponent, Play));
-						CreateCallTimelineFunction(TimelineNode, SourceGraph, PlayName, TimelineVarPin, PlayPin);
-					}
+			// Get the variable output pin
+			UEdGraphPin* TimelineVarPin = GetTimelineNode->FindPin(TimelineNameString);
 
-					if(bPlayFromStartPinConnected)
-					{
-						static FName PlayFromStartName(GET_FUNCTION_NAME_CHECKED(UTimelineComponent, PlayFromStart));
-						CreateCallTimelineFunction(TimelineNode, SourceGraph, PlayFromStartName, TimelineVarPin, PlayFromStartPin);
-					}
+			// This might fail if this is the first compile after adding the timeline (property doesn't exist yet) - in that case, manually add the output pin
+			if(TimelineVarPin == NULL)
+			{
+				TimelineVarPin = GetTimelineNode->CreatePin(EGPD_Output, Schema->PC_Object, TEXT(""), UTimelineComponent::StaticClass(), false, false, TimelineNode->TimelineName.ToString());
+			}
 
-					if(bStopPinConnected)
-					{
-						static FName StopName(GET_FUNCTION_NAME_CHECKED(UTimelineComponent, Stop));
-						CreateCallTimelineFunction(TimelineNode, SourceGraph, StopName, TimelineVarPin, StopPin);
-					}
+			if(bPlayPinConnected)
+			{
+				static FName PlayName(GET_FUNCTION_NAME_CHECKED(UTimelineComponent, Play));
+				CreateCallTimelineFunction(TimelineNode, SourceGraph, PlayName, TimelineVarPin, PlayPin);
+			}
 
-					if (bReversePinConnected)
-					{
-						static FName ReverseName(GET_FUNCTION_NAME_CHECKED(UTimelineComponent, Reverse));
-						CreateCallTimelineFunction(TimelineNode, SourceGraph, ReverseName, TimelineVarPin, ReversePin);
-					}
+			if(bPlayFromStartPinConnected)
+			{
+				static FName PlayFromStartName(GET_FUNCTION_NAME_CHECKED(UTimelineComponent, PlayFromStart));
+				CreateCallTimelineFunction(TimelineNode, SourceGraph, PlayFromStartName, TimelineVarPin, PlayFromStartPin);
+			}
 
-					if (bReverseFromEndPinConnected)
-					{
-						static FName ReverseFromEndName(GET_FUNCTION_NAME_CHECKED(UTimelineComponent, ReverseFromEnd));
-						CreateCallTimelineFunction(TimelineNode, SourceGraph, ReverseFromEndName, TimelineVarPin, ReverseFromEndPin);
-					}
+			if(bStopPinConnected)
+			{
+				static FName StopName(GET_FUNCTION_NAME_CHECKED(UTimelineComponent, Stop));
+				CreateCallTimelineFunction(TimelineNode, SourceGraph, StopName, TimelineVarPin, StopPin);
+			}
 
-					if (bSetNewTimePinConnected)
-					{
-						UEdGraphPin* NewTimePin = TimelineNode->GetNewTimePin();
+			if (bReversePinConnected)
+			{
+				static FName ReverseName(GET_FUNCTION_NAME_CHECKED(UTimelineComponent, Reverse));
+				CreateCallTimelineFunction(TimelineNode, SourceGraph, ReverseName, TimelineVarPin, ReversePin);
+			}
 
-						static FName SetNewTimeName(GET_FUNCTION_NAME_CHECKED(UTimelineComponent, SetNewTime));
-						UK2Node_CallFunction* CallNode = CreateCallTimelineFunction(TimelineNode, SourceGraph, SetNewTimeName, TimelineVarPin, SetTimePin);
+			if (bReverseFromEndPinConnected)
+			{
+				static FName ReverseFromEndName(GET_FUNCTION_NAME_CHECKED(UTimelineComponent, ReverseFromEnd));
+				CreateCallTimelineFunction(TimelineNode, SourceGraph, ReverseFromEndName, TimelineVarPin, ReverseFromEndPin);
+			}
 
-						if (CallNode && NewTimePin)
-						{
-							UEdGraphPin* InputPin = CallNode->FindPinChecked(TEXT("NewTime"));
-							MovePinLinksToIntermediate(*NewTimePin, *InputPin);
-						}
-					}
-				}
+			if (bSetNewTimePinConnected)
+			{
+				UEdGraphPin* NewTimePin = TimelineNode->GetNewTimePin();
 
-				// Create event to call on each update
-				UFunction* EventSigFunc = UTimelineComponent::GetTimelineEventSignature();
+				static FName SetNewTimeName(GET_FUNCTION_NAME_CHECKED(UTimelineComponent, SetNewTime));
+				UK2Node_CallFunction* CallNode = CreateCallTimelineFunction(TimelineNode, SourceGraph, SetNewTimeName, TimelineVarPin, SetTimePin);
 
-				// Create event nodes for any event tracks
-				for(int32 EventTrackIdx=0; EventTrackIdx<Timeline->EventTracks.Num(); EventTrackIdx++)
+				if (CallNode && NewTimePin)
 				{
-					FName EventTrackName =  Timeline->EventTracks[EventTrackIdx].TrackName;
-					CreatePinEventNodeForTimelineFunction(TimelineNode, SourceGraph, Timeline->GetEventTrackFunctionName(EventTrackIdx), EventTrackName.ToString(), EventSigFunc->GetFName());
+					UEdGraphPin* InputPin = CallNode->FindPinChecked(TEXT("NewTime"));
+					MovePinLinksToIntermediate(*NewTimePin, *InputPin);
 				}
+			}
+		}
 
-				// Generate Update Pin Event Node
-				CreatePinEventNodeForTimelineFunction(TimelineNode, SourceGraph, Timeline->GetUpdateFunctionName(), TEXT("Update"), EventSigFunc->GetFName());
+		// Create event to call on each update
+		UFunction* EventSigFunc = UTimelineComponent::GetTimelineEventSignature();
 
-				// Generate Finished Pin Event Node
-				CreatePinEventNodeForTimelineFunction(TimelineNode, SourceGraph, Timeline->GetFinishedFunctionName(), TEXT("Finished"), EventSigFunc->GetFName());
-			}			
-		}
-		else if( UK2Node_VariableGet* VarNode = Cast<UK2Node_VariableGet>( SourceGraph->Nodes[ChildIndex] ))
+		// Create event nodes for any event tracks
+		for(int32 EventTrackIdx=0; EventTrackIdx<Timeline->EventTracks.Num(); EventTrackIdx++)
 		{
-			// Check for Timeline Variable Get Nodes
-			UEdGraphPin* const ValuePin = VarNode->GetValuePin();
-
-			if( ValuePin && ValuePin->LinkedTo.Num() > 0  )
-			{
-				UClass* ValueClass = ValuePin->PinType.PinSubCategoryObject.IsValid() ? Cast<UClass>( ValuePin->PinType.PinSubCategoryObject.Get() ) : nullptr;
-				if( ValueClass == UTimelineComponent::StaticClass() )
-				{
-					if( UTimelineTemplate* TimelineTemplate = Blueprint->FindTimelineTemplateByVariableName( FName( *ValuePin->PinName ) ))
-					{
-						TimelineTemplate->bValidatedAsWired = true;
-					}
-				}
-			}
+			FName EventTrackName =  Timeline->EventTracks[EventTrackIdx].TrackName;
+			CreatePinEventNodeForTimelineFunction(TimelineNode, SourceGraph, Timeline->GetEventTrackFunctionName(EventTrackIdx), EventTrackName.ToString(), EventSigFunc->GetFName());
 		}
+
+		// Generate Update Pin Event Node
+		CreatePinEventNodeForTimelineFunction(TimelineNode, SourceGraph, Timeline->GetUpdateFunctionName(), TEXT("Update"), EventSigFunc->GetFName());
+
+		// Generate Finished Pin Event Node
+		CreatePinEventNodeForTimelineFunction(TimelineNode, SourceGraph, Timeline->GetFinishedFunctionName(), TEXT("Finished"), EventSigFunc->GetFName());
 	}
 }
 
