FKismetCppBackend: CustomThunk functions are handled in generated code. A lot of improvements - ongoing task

diff --git a/Engine/Source/Editor/KismetCompiler/Private/KismetCompilerBackend.h b/Engine/Source/Editor/KismetCompiler/Private/KismetCompilerBackend.h
--- a/Engine/Source/Editor/KismetCompiler/Private/KismetCompilerBackend.h
+++ b/Engine/Source/Editor/KismetCompiler/Private/KismetCompilerBackend.h
@@ -57,6 +57,7 @@ protected:
 	UScriptStruct* RotatorStruct;
 	UScriptStruct* TransformStruct;
 	UScriptStruct* LatentInfoStruct;
+	UScriptStruct* LinearColorStruct;
 public:
 	FStringOutputDevice Header;
 	FStringOutputDevice Body;
@@ -80,6 +81,7 @@ public:
 		VectorStruct = Z_Construct_UScriptStruct_UObject_FVector();
 		RotatorStruct = FindObjectChecked<UScriptStruct>(UObject::StaticClass(), TEXT("Rotator"));
 		TransformStruct = FindObjectChecked<UScriptStruct>(UObject::StaticClass(), TEXT("Transform"));
+		LinearColorStruct = FindObjectChecked<UScriptStruct>(UObject::StaticClass(), TEXT("LinearColor"));
 		LatentInfoStruct = FLatentActionInfo::StaticStruct();
 	}
 
diff --git a/Engine/Source/Editor/KismetCompiler/Private/KismetCompilerCppBackend.cpp b/Engine/Source/Editor/KismetCompiler/Private/KismetCompilerCppBackend.cpp
--- a/Engine/Source/Editor/KismetCompiler/Private/KismetCompilerCppBackend.cpp
+++ b/Engine/Source/Editor/KismetCompiler/Private/KismetCompilerCppBackend.cpp
@@ -18,7 +18,7 @@ FString FKismetCppBackend::TermToText(FBPTerminal* Term, UProperty* CoerceProper
 	if (Term->bIsLiteral)
 	{
 		//@TODO: Must have a coercion type if it's a literal, because the symbol table isn't plumbed in here and the literals don't carry type information either, yay!
-		check(CoerceProperty);
+		ensure(CoerceProperty);
 
 		if (CoerceProperty->IsA(UStrProperty::StaticClass()))
 		{
@@ -86,6 +86,12 @@ FString FKismetCppBackend::TermToText(FBPTerminal* Term, UProperty* CoerceProper
 				return FString::Printf(TEXT("FTransform( FQuat(%f,%f,%f,%f), FVector(%f,%f,%f), FVector(%f,%f,%f) )"),
 					Rot.X, Rot.Y, Rot.Z, Rot.W, Translation.X, Translation.Y, Translation.Z, Scale.X, Scale.Y, Scale.Z);
 			}
+			else if (StructProperty->Struct == LinearColorStruct)
+			{
+				FLinearColor LinearColor;
+				LinearColor.InitFromString(Term->Name);
+				return FString::Printf(TEXT("FLinearColor(%f,%f,%f,%f)"), LinearColor.R, LinearColor.G, LinearColor.B, LinearColor.A);
+			}
 			else
 			{
 				//@todo:  This needs to be more robust, since import text isn't really proper for struct construction.
@@ -122,7 +128,10 @@ FString FKismetCppBackend::TermToText(FBPTerminal* Term, UProperty* CoerceProper
 				}
 				else
 				{
-					return FString::Printf(TEXT("FindObject<UObject>(ANY_PACKAGE, TEXT(\"%s\"))"), *(Term->ObjectLiteral->GetPathName()));
+					UObjectPropertyBase* ObjectCoerceProperty = CastChecked<UObjectPropertyBase>(CoerceProperty);
+					UClass* FoundClass = ObjectCoerceProperty ? ObjectCoerceProperty->PropertyClass : NULL;
+					FString ClassString = FoundClass ? (FString(FoundClass->GetPrefixCPP()) + FoundClass->GetName()) : TEXT("UObject");
+					return FString::Printf(TEXT("FindObject<%s>(ANY_PACKAGE, TEXT(\"%s\"))"), *ClassString, *(Term->ObjectLiteral->GetPathName()));
 				}
 			}
 			else
@@ -202,6 +211,155 @@ FString FKismetCppBackend::LatentFunctionInfoTermToText(FBPTerminal* Term, FBlue
 	return FString(TEXT("F")) + LatentInfoStruct->GetName() + StructValues;
 }
 
+struct FEmitProperyHelper
+{
+	static void ArrayToString(const TArray<FString>& Array, FString& String)
+	{
+		if (Array.Num())
+		{
+			String += Array[0];
+		}
+		for (int32 i = 1; i < Array.Num(); ++i)
+		{
+			String += TEXT(", ");
+			String += Array[i];
+		}
+	}
+
+	static bool HasAllFlags(uint64 Flags, uint64 FlagsToCheck)
+	{
+		return FlagsToCheck == (Flags & FlagsToCheck);
+	}
+
+	static FString HandleRepNotifyFunc(const UProperty* Property)
+	{
+		check(Property);
+		if (HasAllFlags(Property->PropertyFlags, CPF_Net | CPF_RepNotify))
+		{
+			return FString::Printf(TEXT("ReplicatedUsing=%s"), *Property->RepNotifyFunc.ToString());
+		}
+		else if (HasAllFlags(Property->PropertyFlags, CPF_Net))
+		{
+			return TEXT("Replicated");
+		}
+		return FString();
+	}
+
+	static TArray<FString> ProperyFlagsToTags(uint64 Flags)
+	{
+		TArray<FString> Tags;
+
+		// EDIT FLAGS
+		if (HasAllFlags(Flags, (CPF_Edit | CPF_EditConst | CPF_DisableEditOnInstance)))
+		{
+			Tags.Emplace(TEXT("VisibleDefaultsOnly"));
+		}
+		else if (HasAllFlags(Flags, (CPF_Edit | CPF_EditConst | CPF_DisableEditOnTemplate)))
+		{
+			Tags.Emplace(TEXT("VisibleInstanceOnly"));
+		}
+		else if (HasAllFlags(Flags, (CPF_Edit | CPF_EditConst)))
+		{
+			Tags.Emplace(TEXT("VisibleAnywhere"));
+		}
+		else if (HasAllFlags(Flags, (CPF_Edit | CPF_DisableEditOnInstance)))
+		{
+			Tags.Emplace(TEXT("EditDefaultsOnly"));
+		}
+		else if (HasAllFlags(Flags, (CPF_Edit | CPF_DisableEditOnTemplate)))
+		{
+			Tags.Emplace(TEXT("EditInstanceOnly"));
+		}
+		else if (HasAllFlags(Flags, (CPF_Edit)))
+		{
+			Tags.Emplace(TEXT("EditAnywhere"));
+		}
+
+		// BLUEPRINT EDIT
+		if (HasAllFlags(Flags, (CPF_BlueprintVisible | CPF_BlueprintReadOnly)))
+		{
+			Tags.Emplace(TEXT("BlueprintReadOnly"));
+		}
+		else if (HasAllFlags(Flags, (CPF_BlueprintVisible)))
+		{
+			Tags.Emplace(TEXT("BlueprintReadWrite"));
+		}
+
+		// CONFIG
+		if (HasAllFlags(Flags, (CPF_GlobalConfig | CPF_Config)))
+		{
+			Tags.Emplace(TEXT("GlobalConfig"));
+		}
+		else if (HasAllFlags(Flags, (CPF_Config)))
+		{
+			Tags.Emplace(TEXT("Config"));
+		}
+
+		// OTHER
+#ifdef HANDLE_CPF_TAG
+		static_assert(false, "Macro HANDLE_CPF_TAG redefinition.");
+#endif
+#define HANDLE_CPF_TAG(TagName, CheckedFlags) if (HasAllFlags(Flags, (CheckedFlags))) { Tags.Emplace(TagName); }
+		HANDLE_CPF_TAG(TEXT("Localized"), CPF_Localized | CPF_BlueprintReadOnly)
+			HANDLE_CPF_TAG(TEXT("Transient"), CPF_Transient)
+			HANDLE_CPF_TAG(TEXT("DuplicateTransient"), CPF_DuplicateTransient)
+			HANDLE_CPF_TAG(TEXT("TextExportTransient"), CPF_TextExportTransient)
+			HANDLE_CPF_TAG(TEXT("NonPIEDuplicateTransient"), CPF_NonPIEDuplicateTransient)
+			HANDLE_CPF_TAG(TEXT("Export"), CPF_ExportObject)
+			HANDLE_CPF_TAG(TEXT("NoClear"), CPF_NoClear)
+			HANDLE_CPF_TAG(TEXT("EditFixedSize"), CPF_EditFixedSize)
+			HANDLE_CPF_TAG(TEXT("NotReplicated"), CPF_RepSkip)
+			HANDLE_CPF_TAG(TEXT("RepRetry"), CPF_RepRetry)
+			HANDLE_CPF_TAG(TEXT("Interp"), CPF_Edit | CPF_BlueprintVisible | CPF_Interp)
+			HANDLE_CPF_TAG(TEXT("NonTransactional"), CPF_NonTransactional)
+			HANDLE_CPF_TAG(TEXT("Instanced"), CPF_EditInline | CPF_ExportObject | CPF_InstancedReference)
+			HANDLE_CPF_TAG(TEXT("BlueprintAssignable"), CPF_BlueprintAssignable)
+			HANDLE_CPF_TAG(TEXT("BlueprintCallable"), CPF_BlueprintCallable)
+			HANDLE_CPF_TAG(TEXT("BlueprintAuthorityOnly"), CPF_BlueprintAuthorityOnly)
+			HANDLE_CPF_TAG(TEXT("AssetRegistrySearchable"), CPF_AssetRegistrySearchable)
+			HANDLE_CPF_TAG(TEXT("SimpleDisplay"), CPF_SimpleDisplay)
+			HANDLE_CPF_TAG(TEXT("AdvancedDisplay"), CPF_AdvancedDisplay)
+			HANDLE_CPF_TAG(TEXT("SaveGame"), CPF_SaveGame)
+#undef HANDLE_CPF_TAG
+
+			return Tags;
+	}
+
+	static FString HandleMetaData(const UProperty* Property)
+	{
+		FString MetaDataStr;
+
+		check(Property);
+		UPackage* Package = Property->GetOutermost();
+		check(Package);
+		const UMetaData* MetaData = Package->GetMetaData();
+		check(MetaData);
+		const TMap<FName, FString>* ValuesMap = MetaData->ObjectMetaDataMap.Find(Property);
+		if (ValuesMap && ValuesMap->Num())
+		{
+			TArray<FString> MetaDataStrings;
+			for (auto& Pair : *ValuesMap)
+			{
+				if (!Pair.Value.IsEmpty())
+				{
+					MetaDataStrings.Emplace(FString::Printf(TEXT("%s=\"%s\""), *Pair.Key.ToString(), *Pair.Value));
+				}
+				else
+				{
+					MetaDataStrings.Emplace(Pair.Key.ToString());
+				}
+			}
+			MetaDataStrings.Remove(FString());
+			if (MetaDataStrings.Num())
+			{
+				MetaDataStr += TEXT("meta=(");
+				ArrayToString(MetaDataStrings, MetaDataStr);
+				MetaDataStr += TEXT(")");
+			}
+		}
+		return MetaDataStr;
+	}
+};
 void FKismetCppBackend::EmitClassProperties(FStringOutputDevice& Target, UClass* SourceClass)
 {
 	//@TODO: There is a lot of complexity when emitting properties, this code should be copied from FNativeClassHeaderGenerator::ExportProperties (or that code made commonly available in editor mode)
@@ -210,8 +368,20 @@ void FKismetCppBackend::EmitClassProperties(FStringOutputDevice& Target, UClass*
 	for (TFieldIterator<UProperty> It(SourceClass, EFieldIteratorFlags::ExcludeSuper); It; ++It)
 	{
 		UProperty* Property = *It;
+		check(Property);
 
-		Emit(Header, TEXT("\n\tUPROPERTY()\n"));
+		Emit(Header, TEXT("\n\tUPROPERTY("));
+		{
+			TArray<FString> Tags = FEmitProperyHelper::ProperyFlagsToTags(Property->PropertyFlags);
+			Tags.Emplace(FEmitProperyHelper::HandleRepNotifyFunc(Property));
+			Tags.Emplace(FEmitProperyHelper::HandleMetaData(Property));
+			Tags.Remove(FString());
+
+			FString AllTags;
+			FEmitProperyHelper::ArrayToString(Tags, AllTags);
+			Emit(Header, *AllTags);
+		}
+		Emit(Header, TEXT(")\n"));
 		Emit(Header, TEXT("\t"));
 		Property->ExportCppDeclaration(Target, EExportedDeclaration::Member);
 		Emit(Header, TEXT(";\n"));
@@ -223,12 +393,17 @@ void FKismetCppBackend::GenerateCodeFromClass(UClass* SourceClass, TIndirectArra
 	CppClassName = FString(SourceClass->GetPrefixCPP()) + SourceClass->GetName();
 
 	UClass* SuperClass = SourceClass->GetSuperClass();
+
+	Emit(Header, TEXT("#pragma once\n\n"));
+	Emit(Header, *FString::Printf(TEXT("#include \"%s.generated.h\"\n\n"), *SourceClass->GetName()));
+	Emit(Header, TEXT("UCLASS(Blueprintable)\n"));
+
 	Emit(Header,
 		*FString::Printf(TEXT("class %s : public %s%s\n"), *CppClassName, SuperClass->GetPrefixCPP(), *SuperClass->GetName()));
 	Emit(Header,
 		TEXT("{\n")
 		TEXT("public:\n"));
-	Emit(Header, *FString::Printf(TEXT("\tGENERATED_UCLASS_BODY();\n")));
+	Emit(Header, *FString::Printf(TEXT("\tGENERATED_UCLASS_BODY()\n")));
 
 	EmitClassProperties(Header, SourceClass);
 
@@ -245,6 +420,13 @@ void FKismetCppBackend::GenerateCodeFromClass(UClass* SourceClass, TIndirectArra
 		Emit(Header, TEXT("\n"));
 	}
 
+	Emit(Body, TEXT("#include \"GeneratedCodeHelpers.h\"\n\n"));
+
+	//constructor
+	Emit(Body, *FString::Printf(
+		TEXT("%s::%s(const class FPostConstructInitializeProperties& PCIP) : Super(PCIP) {}\n\n"),
+		*CppClassName, *CppClassName));
+
 	for (int32 i = 0; i < Functions.Num(); ++i)
 	{
 		if (Functions[i].IsValid())
@@ -336,7 +518,11 @@ void FKismetCppBackend::EmitCallStatment(FKismetFunctionContext& FunctionContext
 	{
 		if (Statement.FunctionToCall->HasAnyFunctionFlags(FUNC_Static))
 		{
-			Emit(Body, *FString::Printf(TEXT("U%s::"), *Statement.FunctionToCall->GetOuter()->GetName())); //@TODO: Assuming U prefix but could be A
+			const bool bIsCustomThunk = Statement.FunctionToCall->HasMetaData(TEXT("CustomStructureParam")) || Statement.FunctionToCall->HasMetaData(TEXT("ArrayParm"));
+			FString FullFunctionNamePrefix = bIsCustomThunk
+				? TEXT("FCustomThunkTemplates::")
+				: FString::Printf(TEXT("U%s::"), *Statement.FunctionToCall->GetOuter()->GetName());
+			Emit(Body, *FullFunctionNamePrefix); //@TODO: Assuming U prefix but could be A
 		}
 		else
 		{
@@ -482,8 +668,8 @@ void FKismetCppBackend::EmitMetaCastStatement(FKismetFunctionContext& FunctionCo
 	FString SourceClass		= TermToText(Statement.RHS[1], (UProperty*)(GetDefault<UClassProperty>()));
 	FString Destination		= TermToText(Statement.LHS, (UProperty*)(GetDefault<UClassProperty>()));
 
-	Emit(Body, *FString::Printf(TEXT("\t\t\t%s = ((%s)->IsChildOf(DesiredClass)) ? SourceClass : NULL;\n"), 
-		*Destination, *SourceClass, *DesiredClass, *SourceClass));
+	Emit(Body, *FString::Printf(TEXT("\t\t\t%s = DynamicMetaCast(%s, %s);\n"),
+		*Destination, *DesiredClass, *SourceClass));
 }
 
 void FKismetCppBackend::EmitObjectToBoolStatement(FKismetFunctionContext& FunctionContext, FBlueprintCompiledStatement& Statement)
@@ -708,7 +894,7 @@ void FKismetCppBackend::ConstructFunction(FKismetFunctionContext& FunctionContex
 
 		Emit(Header, TEXT("\t"));
 		Emit(Header, *FString::Printf(*Start, TEXT(""), TEXT("")));
-		Emit(Body, *FString::Printf(*Start, *Class->GetName(), TEXT("::")));
+		Emit(Body, *FString::Printf(*Start, *CppClassName, TEXT("::")));
 
 		for (int32 i = 0; i < ArgumentList.Num(); ++i)
 		{
diff --git a/Engine/Source/Runtime/Engine/Public/GeneratedCodeHelpers.h b/Engine/Source/Runtime/Engine/Public/GeneratedCodeHelpers.h
new file mode 100644
--- /dev/null
+++ b/Engine/Source/Runtime/Engine/Public/GeneratedCodeHelpers.h
@@ -0,0 +1,123 @@
+// Copyright 1998-2014 Epic Games, Inc. All Rights Reserved.
+
+#pragma once
+
+#include "Kismet/KismetArrayLibrary.h"
+#include "Kismet/KismetSystemLibrary.h"
+#include "Kismet/DataTableFunctionLibrary.h"
+
+FORCEINLINE UClass* DynamicMetaCast(const UClass* DesiredClass, UClass* SourceClass)
+{
+	return ((SourceClass)->IsChildOf(DesiredClass)) ? SourceClass : NULL;
+}
+
+struct FCustomThunkTemplates
+{
+	//Replacements for CustomThunk functions from UKismetArrayLibrary
+
+	template<typename T>
+	static int32 Array_Add(TArray<T>& TargetArray, const UArrayProperty* ArrayProperty, const T& NewItem)
+	{
+		return UKismetArrayLibrary::GenericArray_Add(&TargetArray, ArrayProperty, &NewItem);
+	}
+
+	template<typename T>
+	static void Array_Shuffle(TArray<T>& TargetArray, const UArrayProperty* ArrayProperty)
+	{
+		UKismetArrayLibrary::GenericArray_Shuffle(&TargetArray, ArrayProperty);
+	}
+
+	template<typename T>
+	static void Array_Append(TArray<T>& TargetArray, const UArrayProperty* TargetArrayProperty, const TArray<T>& SourceArray, const UArrayProperty* SourceArrayProperty)
+	{
+		UKismetArrayLibrary::GenericArray_Append(&TargetArray, TargetArrayProperty, &SourceArray, SourceArrayProperty);
+	}
+
+	template<typename T>
+	static void Array_Insert(TArray<T>& TargetArray, const UArrayProperty* ArrayProperty, const T& NewItem, int32 Index)
+	{
+		UKismetArrayLibrary::GenericArray_Insert(&TargetArray, ArrayProperty, &NewItem, Index);
+	}
+
+	template<typename T>
+	static void Array_Remove(TArray<T>& TargetArray, const UArrayProperty* ArrayProperty, int32 IndexToRemove)
+	{
+		UKismetArrayLibrary::GenericArray_Remove(&TargetArray, ArrayProperty, IndexToRemove);
+	}
+
+	template<typename T>
+	static int32 Array_Find(const TArray<T>& TargetArray, const UArrayProperty* ArrayProperty, const T& ItemToFind)
+	{
+		return TargetArray.Find(ItemToFind);
+	}
+
+	template<typename T>
+	static bool Array_RemoveItem(TArray<T>& TargetArray, const UArrayProperty* ArrayProperty, const T& Item)
+	{
+		return TargetArray.Remove(Item);
+	}
+
+	template<typename T>
+	static void Array_Clear(TArray<T>& TargetArray, const UArrayProperty* ArrayProperty)
+	{
+		UKismetArrayLibrary::GenericArray_Clear(&TargetArray, ArrayProperty);
+	}
+
+	template<typename T>
+	static void Array_Resize(TArray<T>& TargetArray, const UArrayProperty* ArrayProperty, int32 Size)
+	{
+		UKismetArrayLibrary::GenericArray_Resize(&TargetArray, ArrayProperty, Size);
+	}
+
+	template<typename T>
+	static int32 Array_Length(const TArray<T>& TargetArray, const UArrayProperty* ArrayProperty)
+	{
+		return UKismetArrayLibrary::GenericArray_Length(&TargetArray, ArrayProperty);
+	}
+
+	template<typename T>
+	static int32 Array_LastIndex(const TArray<T>& TargetArray, const UArrayProperty* ArrayProperty)
+	{
+		return UKismetArrayLibrary::GenericArray_Length(&TargetArray, ArrayProperty);
+	}
+
+	template<typename T>
+	static void Array_Get(TArray<T>& TargetArray, const UArrayProperty* ArrayProperty, int32 Index, T& Item)
+	{
+		UKismetArrayLibrary::GenericArray_Get(&TargetArray, ArrayProperty, Index, &Item);
+	}
+
+	template<typename T>
+	static void Array_Set(TArray<T>& TargetArray, const UArrayProperty* ArrayProperty, int32 Index, const T& Item, bool bSizeToFit)
+	{
+		UKismetArrayLibrary::GenericArray_Set(&TargetArray, ArrayProperty, Index, &Item, bSizeToFit);
+	}
+
+	template<typename T>
+	static bool Array_Contains(const TArray<T>& TargetArray, const UArrayProperty* ArrayProperty, const T& ItemToFind)
+	{
+		return Array_Find(TargetArray, ArrayProperty, ItemToFind) > 0;
+	}
+
+	template<typename T>
+	static void SetArrayPropertyByName(UObject* Object, FName PropertyName, TArray<T>& Value)
+	{
+		UKismetArrayLibrary::GenericArray_SetArrayPropertyByName(Object, PropertyName, &Value);
+	}
+
+	//Replacements for CustomThunk functions from UDataTableFunctionLibrary
+
+	template<typename T>
+	static bool GetDataTableRowFromName(UDataTable* Table, FName RowName, T& OutRow)
+	{
+		return UDataTableFunctionLibrary::Generic_GetDataTableRowFromName(Table, RowName, &OutRow);
+	}
+
+	//Replacements for CustomThunk functions from UKismetSystemLibrary
+
+	template<typename T>
+	static void SetStructurePropertyByName(UObject* Object, FName PropertyName, const T& Value)
+	{
+		return UKismetSystemLibrary::Generic_SetStructurePropertyByName(Object, PropertyName, &Value);
+	}
+};
\ No newline at end of file
