Fix for yet another edge case on the timeline culling.
#TTP 346668 - Blueprint nodes being possibly optimized out
#Branch UE4
#Proj KismetCompiler

#Change added testing for references in addition to the timeline being wired to ensure that use variable get use cases are covered.

diff --git a/Engine/Source/Editor/KismetCompiler/Private/KismetCompiler.cpp b/Engine/Source/Editor/KismetCompiler/Private/KismetCompiler.cpp
--- a/Engine/Source/Editor/KismetCompiler/Private/KismetCompiler.cpp
+++ b/Engine/Source/Editor/KismetCompiler/Private/KismetCompiler.cpp
@@ -1895,9 +1895,10 @@ void FKismetCompilerContext::ExpandTimelineNodes(UEdGraph* SourceGraph)
 		// Set the timeline template as wired/not wired for component pruning later
 		const bool bWiredIn = bPlayPinConnected || bPlayFromStartPinConnected || bStopPinConnected || bReversePinConnected || bReverseFromEndPinConnected || bSetNewTimePinConnected;
 		const bool bWiredOut = bUpdatePinConnected || bFinishedPinConnected;
-		const bool bPlayWired = Timeline->bAutoPlay || ( TimelinePlayNodes.Find( TimelineNode->TimelineName ) != INDEX_NONE );
+		const bool bPlayWired = Timeline->bAutoPlay;
+		const bool bReferenced = TimelinePlayNodes.Find( TimelineNode->TimelineName ) != INDEX_NONE;
 
-		Timeline->bValidatedAsWired = bWiredIn || ( bPlayWired && bWiredOut );
+		Timeline->bValidatedAsWired = bWiredIn || bReferenced || ( bPlayWired && bWiredOut );
 
 		// Only create nodes for play/stop if they are actually connected - otherwise we get a 'unused node being pruned' warning
 		if(bWiredIn)
