BP: FKismetCppBackend::EmitBindDelegateStatement - fix

diff --git a/Engine/Source/Editor/KismetCompiler/Private/KismetCompilerBackend.h b/Engine/Source/Editor/KismetCompiler/Private/KismetCompilerBackend.h
--- a/Engine/Source/Editor/KismetCompiler/Private/KismetCompilerBackend.h
+++ b/Engine/Source/Editor/KismetCompiler/Private/KismetCompilerBackend.h
@@ -62,7 +62,7 @@ public:
 	FStringOutputDevice Header;
 	FStringOutputDevice Body;
 protected:
-	FString TermToText(FBPTerminal* Term, UProperty* SourceProperty = NULL);
+	FString TermToText(const FBPTerminal* Term, const UProperty* SourceProperty = NULL);
 	FString LatentFunctionInfoTermToText(FBPTerminal* Term, FBlueprintCompiledStatement* TargetLabel);
 
 	int32 StatementToStateIndex(FKismetFunctionContext& FunctionContext, FBlueprintCompiledStatement* Statement)
diff --git a/Engine/Source/Editor/KismetCompiler/Private/KismetCompilerCppBackend.cpp b/Engine/Source/Editor/KismetCompiler/Private/KismetCompilerCppBackend.cpp
--- a/Engine/Source/Editor/KismetCompiler/Private/KismetCompilerCppBackend.cpp
+++ b/Engine/Source/Editor/KismetCompiler/Private/KismetCompilerCppBackend.cpp
@@ -13,7 +13,7 @@
 //////////////////////////////////////////////////////////////////////////
 // FKismetCppBackend
 
-FString FKismetCppBackend::TermToText(FBPTerminal* Term, UProperty* CoerceProperty)
+FString FKismetCppBackend::TermToText(const FBPTerminal* Term, const UProperty* CoerceProperty)
 {
 	if (Term->bIsLiteral)
 	{
@@ -34,7 +34,7 @@ FString FKismetCppBackend::TermToText(FBPTerminal* Term, UProperty* CoerceProper
 			int32 Value = FCString::Atoi(*(Term->Name));
 			return FString::Printf(TEXT("%d"), Value);
 		}
-		else if (UByteProperty* ByteProperty = Cast<UByteProperty>(CoerceProperty))
+		else if (auto ByteProperty = Cast<const UByteProperty>(CoerceProperty))
 		{
 			// The PinSubCategoryObject check is to allow enum literals communicate with byte properties as literals
 			if (ByteProperty->Enum != NULL ||
@@ -48,17 +48,17 @@ FString FKismetCppBackend::TermToText(FBPTerminal* Term, UProperty* CoerceProper
 				return FString::Printf(TEXT("%u"), Value);
 			}
 		}
-		else if (UBoolProperty* BoolProperty = Cast<UBoolProperty>(CoerceProperty))
+		else if (auto BoolProperty = Cast<const UBoolProperty>(CoerceProperty))
 		{
 			bool bValue = Term->Name.ToBool();
 			return *(FName::GetEntry(bValue ? NAME_TRUE : NAME_FALSE)->GetPlainNameString());
 		}
-		else if (UNameProperty* NameProperty = Cast<UNameProperty>(CoerceProperty))
+		else if (auto NameProperty = Cast<const UNameProperty>(CoerceProperty))
 		{
 			FName LiteralName(*(Term->Name));
 			return FString::Printf(TEXT("FName(TEXT(\"%s\")"), *(LiteralName.ToString()));
 		}
-		else if (UStructProperty* StructProperty = Cast<UStructProperty>(CoerceProperty))
+		else if (auto StructProperty = Cast<const UStructProperty>(CoerceProperty))
 		{
 			if (StructProperty->Struct == VectorStruct)
 			{
@@ -98,7 +98,7 @@ FString FKismetCppBackend::TermToText(FBPTerminal* Term, UProperty* CoerceProper
 				return FString(TEXT("F")) + StructProperty->Struct->GetName() + Term->Name;
 			}
 		}
-		else if (UClassProperty* ClassProperty = Cast<UClassProperty>(CoerceProperty))
+		else if (auto ClassProperty = Cast<const UClassProperty>(CoerceProperty))
 		{
 			return FString::Printf(TEXT("%s::StaticClass()"), *(Term->Name)); //@TODO: Need the correct class prefix
 		}
@@ -122,13 +122,13 @@ FString FKismetCppBackend::TermToText(FBPTerminal* Term, UProperty* CoerceProper
 			}
 			else if (Term->ObjectLiteral)
 			{
-				if (UClass* LiteralClass = Cast<UClass>(Term->ObjectLiteral))
+				if (auto LiteralClass = Cast<const UClass>(Term->ObjectLiteral))
 				{
 					return FString::Printf(TEXT("%s::StaticClass()"), *LiteralClass->GetName()); //@TODO: Need the correct class prefix
 				}
 				else
 				{
-					UObjectPropertyBase* ObjectCoerceProperty = CastChecked<UObjectPropertyBase>(CoerceProperty);
+					auto ObjectCoerceProperty = CastChecked<const UObjectPropertyBase>(CoerceProperty);
 					UClass* FoundClass = ObjectCoerceProperty ? ObjectCoerceProperty->PropertyClass : NULL;
 					FString ClassString = FoundClass ? (FString(FoundClass->GetPrefixCPP()) + FoundClass->GetName()) : TEXT("UObject");
 					return FString::Printf(TEXT("FindObject<%s>(ANY_PACKAGE, TEXT(\"%s\"))"), *ClassString, *(Term->ObjectLiteral->GetPathName()));
@@ -700,8 +700,8 @@ void FKismetCppBackend::EmitBindDelegateStatement(FKismetFunctionContext& Functi
 {
 	check(2 == Statement.RHS.Num());
 	const FString Delegate = TermToText(Statement.LHS);
-	const FString NameTerm = TermToText(Statement.RHS[0]);
-	const FString ObjectTerm = TermToText(Statement.RHS[1]);
+	const FString NameTerm = TermToText(Statement.RHS[0], GetDefault<UNameProperty>());
+	const FString ObjectTerm = TermToText(Statement.RHS[1], GetDefault<UObjectProperty>());
 
 	Emit(Body, 
 		*FString::Printf(
diff --git a/Engine/Source/Editor/KismetCompiler/Public/BPTerminal.h b/Engine/Source/Editor/KismetCompiler/Public/BPTerminal.h
--- a/Engine/Source/Editor/KismetCompiler/Public/BPTerminal.h
+++ b/Engine/Source/Editor/KismetCompiler/Public/BPTerminal.h
@@ -42,6 +42,7 @@ struct FBPTerminal
 		, bIsSavePersistent(false)
 		, bIsStructContext(false)
 		, bPassedByReference(false)
+		, Source(NULL)
 		, Context(NULL)
 		, AssociatedVarProperty(NULL)
 		, ObjectLiteral(NULL)
