BP: FKismetCppBackend - a lot of fixes/improvements:
- Proper code for EventDispatcher
- UFUNCTION macro
- Delegate declaration..

diff --git a/Engine/Source/Editor/BlueprintGraph/Private/DelegateNodeHandlers.cpp b/Engine/Source/Editor/BlueprintGraph/Private/DelegateNodeHandlers.cpp
--- a/Engine/Source/Editor/BlueprintGraph/Private/DelegateNodeHandlers.cpp
+++ b/Engine/Source/Editor/BlueprintGraph/Private/DelegateNodeHandlers.cpp
@@ -92,7 +92,7 @@ struct FKCHandlerDelegateHelper
 		check(SelfPin && NetPin && BoundProperty && DelegateNode);
 
 		FBPTerminal* Term = new(Context.VariableReferences) FBPTerminal();
-		Term->CopyFromPin(SelfPin, Context.NetNameMap->MakeValidName(SelfPin));
+		Term->CopyFromPin(SelfPin, BoundProperty->GetName());
 		Term->AssociatedVarProperty = BoundProperty;
 
 		FBPTerminal** pContextTerm = Context.NetMap.Find(NetPin);
diff --git a/Engine/Source/Editor/KismetCompiler/Private/KismetCompilerBackend.h b/Engine/Source/Editor/KismetCompiler/Private/KismetCompilerBackend.h
--- a/Engine/Source/Editor/KismetCompiler/Private/KismetCompilerBackend.h
+++ b/Engine/Source/Editor/KismetCompiler/Private/KismetCompilerBackend.h
@@ -61,8 +61,8 @@ protected:
 public:
 	FStringOutputDevice Header;
 	FStringOutputDevice Body;
-protected:
 	FString TermToText(const FBPTerminal* Term, const UProperty* SourceProperty = NULL);
+protected:
 	FString LatentFunctionInfoTermToText(FBPTerminal* Term, FBlueprintCompiledStatement* TargetLabel);
 
 	int32 StatementToStateIndex(FKismetFunctionContext& FunctionContext, FBlueprintCompiledStatement* Statement)
@@ -85,6 +85,8 @@ public:
 		LatentInfoStruct = FLatentActionInfo::StaticStruct();
 	}
 
+	UEdGraphSchema_K2* GetSchema() const { return Schema; }
+
 	void Emit(FStringOutputDevice& Target, const TCHAR* Message)
 	{
 		Target += Message;
diff --git a/Engine/Source/Editor/KismetCompiler/Private/KismetCompilerCppBackend.cpp b/Engine/Source/Editor/KismetCompiler/Private/KismetCompilerCppBackend.cpp
--- a/Engine/Source/Editor/KismetCompiler/Private/KismetCompilerCppBackend.cpp
+++ b/Engine/Source/Editor/KismetCompiler/Private/KismetCompilerCppBackend.cpp
@@ -10,6 +10,68 @@
 
 #include "DefaultValueHelper.h"
 
+// Generates single if scope. It's condition checks context of given term.
+struct FSafeContextScopedEmmitter
+{
+private:
+	FStringOutputDevice& Body;
+	bool bSafeContextUsed;
+	const TCHAR* CurrentIndent;
+public:
+	FString GetAdditionalIndent() const
+	{
+		return bSafeContextUsed ? TEXT("\t") : FString();
+	}
+
+	bool IsSafeContextUsed() const
+	{
+		return bSafeContextUsed;
+	}
+
+	FSafeContextScopedEmmitter(FStringOutputDevice& InBody, const FBPTerminal* Term, FKismetCppBackend& CppBackend, const TCHAR* InCurrentIndent)
+		: Body(InBody), bSafeContextUsed(false), CurrentIndent(InCurrentIndent)
+	{
+		TArray<FString> SafetyConditions;
+		for (; Term; Term = Term->Context)
+		{
+			if (!Term->bIsStructContext && (Term->Type.PinSubCategory != CppBackend.GetSchema()->PSC_Self))
+			{
+				ensure(!Term->bIsLiteral);
+				SafetyConditions.Add(CppBackend.TermToText(Term));
+			}
+		}
+
+		if (SafetyConditions.Num())
+		{
+			bSafeContextUsed = true;
+			Body += CurrentIndent;
+			Body += TEXT("if (");
+			for (int32 Iter = SafetyConditions.Num() - 1; Iter >= 0; --Iter)
+			{
+				Body += FString(TEXT("IsValid("));
+				Body += SafetyConditions[Iter];
+				Body += FString(TEXT(")"));
+				if (Iter)
+				{
+					Body += FString(TEXT(" && "));
+				}
+			}
+			Body += TEXT(")\n");
+			Body += CurrentIndent;
+			Body += TEXT("{\n");
+		}
+	}
+
+	~FSafeContextScopedEmmitter()
+	{
+		if (bSafeContextUsed)
+		{
+			Body += CurrentIndent;
+			Body += TEXT("}\n");
+		}
+	}
+};
+
 //////////////////////////////////////////////////////////////////////////
 // FKismetCppBackend
 
@@ -100,7 +162,11 @@ FString FKismetCppBackend::TermToText(const FBPTerminal* Term, const UProperty*
 		}
 		else if (auto ClassProperty = Cast<const UClassProperty>(CoerceProperty))
 		{
-			return FString::Printf(TEXT("%s::StaticClass()"), *(Term->Name)); //@TODO: Need the correct class prefix
+			if (auto FoundClass = Cast<const UClass>(Term->ObjectLiteral))
+			{
+				return FString::Printf(TEXT("%s%s::StaticClass()"), *FoundClass->GetPrefixCPP(), *FoundClass->GetName());
+			}
+			return FString(TEXT("NULL"));
 		}
 		else if (CoerceProperty->IsA(UDelegateProperty::StaticClass()))
 		{
@@ -124,7 +190,7 @@ FString FKismetCppBackend::TermToText(const FBPTerminal* Term, const UProperty*
 			{
 				if (auto LiteralClass = Cast<const UClass>(Term->ObjectLiteral))
 				{
-					return FString::Printf(TEXT("%s::StaticClass()"), *LiteralClass->GetName()); //@TODO: Need the correct class prefix
+					return FString::Printf(TEXT("%s%s::StaticClass()"), *LiteralClass->GetPrefixCPP(), *LiteralClass->GetName());
 				}
 				else
 				{
@@ -162,10 +228,8 @@ FString FKismetCppBackend::TermToText(const FBPTerminal* Term, const UProperty*
 	else
 	{
 		FString Prefix(TEXT(""));
-		if (Term->Context != NULL)
+		if ((Term->Context != NULL) && (Term->Context->Name != Schema->PSC_Self))
 		{
-			//@TODO: Badness, could be a self reference wired to another instance!
-			check(Term->Context->Name != Schema->PSC_Self);
 			Prefix = TermToText(Term->Context);
 
 			if (Term->Context->bIsStructContext)
@@ -211,9 +275,9 @@ FString FKismetCppBackend::LatentFunctionInfoTermToText(FBPTerminal* Term, FBlue
 	return FString(TEXT("F")) + LatentInfoStruct->GetName() + StructValues;
 }
 
-struct FEmitProperyHelper
+struct FEmitHelper
 {
-	static void ArrayToString(const TArray<FString>& Array, FString& String)
+	static void ArrayToString(const TArray<FString>& Array, FString& String, const TCHAR* Separator)
 	{
 		if (Array.Num())
 		{
@@ -245,6 +309,46 @@ struct FEmitProperyHelper
 		return FString();
 	}
 
+	static FString HandleMetaData(const UField* Field)
+	{
+		FString MetaDataStr;
+
+		check(Field);
+		UPackage* Package = Field->GetOutermost();
+		check(Package);
+		const UMetaData* MetaData = Package->GetMetaData();
+		check(MetaData);
+		const TMap<FName, FString>* ValuesMap = MetaData->ObjectMetaDataMap.Find(Field);
+		if (ValuesMap && ValuesMap->Num())
+		{
+			TArray<FString> MetaDataStrings;
+			for (auto& Pair : *ValuesMap)
+			{
+				if (!Pair.Value.IsEmpty())
+				{
+					MetaDataStrings.Emplace(FString::Printf(TEXT("%s=\"%s\""), *Pair.Key.ToString(), *Pair.Value));
+				}
+				else
+				{
+					MetaDataStrings.Emplace(Pair.Key.ToString());
+				}
+			}
+			MetaDataStrings.Remove(FString());
+			if (MetaDataStrings.Num())
+			{
+				MetaDataStr += TEXT("meta=(");
+				ArrayToString(MetaDataStrings, MetaDataStr, TEXT(", "));
+				MetaDataStr += TEXT(")");
+			}
+		}
+		return MetaDataStr;
+	}
+
+#ifdef HANDLE_CPF_TAG
+	static_assert(false, "Macro HANDLE_CPF_TAG redefinition.");
+#endif
+#define HANDLE_CPF_TAG(TagName, CheckedFlags) if (HasAllFlags(Flags, (CheckedFlags))) { Tags.Emplace(TagName); }
+
 	static TArray<FString> ProperyFlagsToTags(uint64 Flags)
 	{
 		TArray<FString> Tags;
@@ -296,74 +400,140 @@ struct FEmitProperyHelper
 		}
 
 		// OTHER
-#ifdef HANDLE_CPF_TAG
-		static_assert(false, "Macro HANDLE_CPF_TAG redefinition.");
-#endif
-#define HANDLE_CPF_TAG(TagName, CheckedFlags) if (HasAllFlags(Flags, (CheckedFlags))) { Tags.Emplace(TagName); }
 		HANDLE_CPF_TAG(TEXT("Localized"), CPF_Localized | CPF_BlueprintReadOnly)
-			HANDLE_CPF_TAG(TEXT("Transient"), CPF_Transient)
-			HANDLE_CPF_TAG(TEXT("DuplicateTransient"), CPF_DuplicateTransient)
-			HANDLE_CPF_TAG(TEXT("TextExportTransient"), CPF_TextExportTransient)
-			HANDLE_CPF_TAG(TEXT("NonPIEDuplicateTransient"), CPF_NonPIEDuplicateTransient)
-			HANDLE_CPF_TAG(TEXT("Export"), CPF_ExportObject)
-			HANDLE_CPF_TAG(TEXT("NoClear"), CPF_NoClear)
-			HANDLE_CPF_TAG(TEXT("EditFixedSize"), CPF_EditFixedSize)
-			HANDLE_CPF_TAG(TEXT("NotReplicated"), CPF_RepSkip)
-			HANDLE_CPF_TAG(TEXT("RepRetry"), CPF_RepRetry)
-			HANDLE_CPF_TAG(TEXT("Interp"), CPF_Edit | CPF_BlueprintVisible | CPF_Interp)
-			HANDLE_CPF_TAG(TEXT("NonTransactional"), CPF_NonTransactional)
-			HANDLE_CPF_TAG(TEXT("Instanced"), CPF_PersistentInstance | CPF_ExportObject | CPF_InstancedReference)
-			HANDLE_CPF_TAG(TEXT("BlueprintAssignable"), CPF_BlueprintAssignable)
-			HANDLE_CPF_TAG(TEXT("BlueprintCallable"), CPF_BlueprintCallable)
-			HANDLE_CPF_TAG(TEXT("BlueprintAuthorityOnly"), CPF_BlueprintAuthorityOnly)
-			HANDLE_CPF_TAG(TEXT("AssetRegistrySearchable"), CPF_AssetRegistrySearchable)
-			HANDLE_CPF_TAG(TEXT("SimpleDisplay"), CPF_SimpleDisplay)
-			HANDLE_CPF_TAG(TEXT("AdvancedDisplay"), CPF_AdvancedDisplay)
-			HANDLE_CPF_TAG(TEXT("SaveGame"), CPF_SaveGame)
+		HANDLE_CPF_TAG(TEXT("Transient"), CPF_Transient)
+		HANDLE_CPF_TAG(TEXT("DuplicateTransient"), CPF_DuplicateTransient)
+		HANDLE_CPF_TAG(TEXT("TextExportTransient"), CPF_TextExportTransient)
+		HANDLE_CPF_TAG(TEXT("NonPIEDuplicateTransient"), CPF_NonPIEDuplicateTransient)
+		HANDLE_CPF_TAG(TEXT("Export"), CPF_ExportObject)
+		HANDLE_CPF_TAG(TEXT("NoClear"), CPF_NoClear)
+		HANDLE_CPF_TAG(TEXT("EditFixedSize"), CPF_EditFixedSize)
+		HANDLE_CPF_TAG(TEXT("NotReplicated"), CPF_RepSkip)
+		HANDLE_CPF_TAG(TEXT("RepRetry"), CPF_RepRetry)
+		HANDLE_CPF_TAG(TEXT("Interp"), CPF_Edit | CPF_BlueprintVisible | CPF_Interp)
+		HANDLE_CPF_TAG(TEXT("NonTransactional"), CPF_NonTransactional)
+		HANDLE_CPF_TAG(TEXT("BlueprintAssignable"), CPF_BlueprintAssignable)
+		HANDLE_CPF_TAG(TEXT("BlueprintCallable"), CPF_BlueprintCallable)
+		HANDLE_CPF_TAG(TEXT("BlueprintAuthorityOnly"), CPF_BlueprintAuthorityOnly)
+		HANDLE_CPF_TAG(TEXT("AssetRegistrySearchable"), CPF_AssetRegistrySearchable)
+		HANDLE_CPF_TAG(TEXT("SimpleDisplay"), CPF_SimpleDisplay)
+		HANDLE_CPF_TAG(TEXT("AdvancedDisplay"), CPF_AdvancedDisplay)
+		HANDLE_CPF_TAG(TEXT("SaveGame"), CPF_SaveGame)
+
+		//TODO:
+		//HANDLE_CPF_TAG(TEXT("Instanced"), CPF_PersistentInstance | CPF_ExportObject | CPF_InstancedReference)
+
+		return Tags;
+	}
+
+	static TArray<FString> FunctionFlagsToTags(uint64 Flags)
+	{
+		TArray<FString> Tags;
+
+		//TODO: BlueprintNativeEvent, BlueprintImplementableEvent
+		//TODO: SealedEvent
+		//TODO: CustomThunk
+		//TODO: Unreliable
+		//TODO: ServiceRequest, ServiceResponse
+		
+		HANDLE_CPF_TAG(TEXT("Exec"), FUNC_Exec)
+		HANDLE_CPF_TAG(TEXT("Server"), FUNC_Net | FUNC_NetServer)
+		HANDLE_CPF_TAG(TEXT("Client"), FUNC_Net | FUNC_NetClient)
+		HANDLE_CPF_TAG(TEXT("NetMulticast"), FUNC_Net | FUNC_NetMulticast)
+		HANDLE_CPF_TAG(TEXT("Reliable"), FUNC_NetReliable)
+		HANDLE_CPF_TAG(TEXT("BlueprintCallable"), FUNC_BlueprintCallable)
+		HANDLE_CPF_TAG(TEXT("BlueprintPure"), FUNC_BlueprintCallable | FUNC_BlueprintPure)
+		HANDLE_CPF_TAG(TEXT("BlueprintAuthorityOnly"), FUNC_BlueprintAuthorityOnly)
+		HANDLE_CPF_TAG(TEXT("BlueprintCosmetic"), FUNC_BlueprintCosmetic)
+		HANDLE_CPF_TAG(TEXT("WithValidation"), FUNC_NetValidate)
+
+		return Tags;
+	}
+
 #undef HANDLE_CPF_TAG
 
-			return Tags;
+	static FString EmitUFuntion(UFunction* Function)
+	{
+		TArray<FString> Tags = FEmitHelper::FunctionFlagsToTags(Function->FunctionFlags);
+		Tags.Emplace(FEmitHelper::HandleMetaData(Function));
+		Tags.Remove(FString());
+
+		FString AllTags;
+		FEmitHelper::ArrayToString(Tags, AllTags, TEXT(", "));
+
+		return FString::Printf(TEXT("UFUNCTION(%s)"), *AllTags);
 	}
 
-	static FString HandleMetaData(const UProperty* Property)
+	static int32 ParseDelegateDetails(UFunction* Signature, FString& OutParametersMacro, FString& OutParamNumberStr)
 	{
-		FString MetaDataStr;
+		int32 ParameterNum = 0;
+		FStringOutputDevice Parameters;
+		for (TFieldIterator<UProperty> PropIt(Signature); PropIt && (PropIt->PropertyFlags & CPF_Parm); ++PropIt)
+		{
+			Parameters += ", ";
+			PropIt->ExportCppDeclaration(Parameters, EExportedDeclaration::MacroParameter, NULL, EPropertyExportCPPFlags::CPPF_CustomTypeName);
+			++ParameterNum;
+		}
 
-		check(Property);
-		UPackage* Package = Property->GetOutermost();
-		check(Package);
-		const UMetaData* MetaData = Package->GetMetaData();
-		check(MetaData);
-		const TMap<FName, FString>* ValuesMap = MetaData->ObjectMetaDataMap.Find(Property);
-		if (ValuesMap && ValuesMap->Num())
+		FString ParamNumberStr;
+		switch (ParameterNum)
 		{
-			TArray<FString> MetaDataStrings;
-			for (auto& Pair : *ValuesMap)
-			{
-				if (!Pair.Value.IsEmpty())
-				{
-					MetaDataStrings.Emplace(FString::Printf(TEXT("%s=\"%s\""), *Pair.Key.ToString(), *Pair.Value));
-				}
-				else
-				{
-					MetaDataStrings.Emplace(Pair.Key.ToString());
-				}
-			}
-			MetaDataStrings.Remove(FString());
-			if (MetaDataStrings.Num())
-			{
-				MetaDataStr += TEXT("meta=(");
-				ArrayToString(MetaDataStrings, MetaDataStr);
-				MetaDataStr += TEXT(")");
-			}
+		case 0: ParamNumberStr = TEXT("");				break;
+		case 1: ParamNumberStr = TEXT("_OneParam");		break;
+		case 2: ParamNumberStr = TEXT("_TwoParams");	break;
+		case 3: ParamNumberStr = TEXT("_ThreeParams");	break;
+		case 4: ParamNumberStr = TEXT("_FourParams");	break;
+		case 5: ParamNumberStr = TEXT("_FiveParams");	break;
+		case 6: ParamNumberStr = TEXT("_SixParams");	break;
+		case 7: ParamNumberStr = TEXT("_SevenParams");	break;
+		case 8: ParamNumberStr = TEXT("_EightParams");	break;
+		default: ParamNumberStr = TEXT("_TooMany");		break;
 		}
-		return MetaDataStr;
+
+		OutParametersMacro = Parameters;
+		OutParamNumberStr = ParamNumberStr;
+		return ParameterNum;
+	}
+
+	static TArray<FString> EmitSinglecastDelegateDeclarations(const TArray<UMulticastDelegateProperty*>& Delegates)
+	{
+		TArray<FString> Results;
+		for (auto It : Delegates)
+		{
+			check(It);
+			auto Signature = It->SignatureFunction;
+			check(Signature);
+
+			FString ParamNumberStr, Parameters;
+			ParseDelegateDetails(Signature, Parameters, ParamNumberStr);
+
+			Results.Add(*FString::Printf(TEXT("DECLARE_DYNAMIC_DELEGATE%s(%s%s)"),
+				*ParamNumberStr, *It->GetCPPType(NULL, EPropertyExportCPPFlags::CPPF_CustomTypeName), *Parameters));
+		}
+		return Results;
+	}
+
+	static TArray<FString> EmitMulticastDelegateDeclarations(UClass* SourceClass)
+	{
+		TArray<FString> Results;
+		for (TFieldIterator<UMulticastDelegateProperty> It(SourceClass, EFieldIteratorFlags::ExcludeSuper); It; ++It)
+		{
+			auto Signature = It->SignatureFunction;
+			check(Signature);
+
+			FString ParamNumberStr, Parameters;
+			ParseDelegateDetails(Signature, Parameters, ParamNumberStr);
+
+			Results.Add(*FString::Printf(TEXT("DECLARE_DYNAMIC_MULTICAST_DELEGATE%s(%s%s)"),
+				*ParamNumberStr, *It->GetCPPType(NULL, EPropertyExportCPPFlags::CPPF_CustomTypeName), *Parameters));
+		}
+
+		return Results;
 	}
 };
+
 void FKismetCppBackend::EmitClassProperties(FStringOutputDevice& Target, UClass* SourceClass)
 {
-	//@TODO: There is a lot of complexity when emitting properties, this code should be copied from FNativeClassHeaderGenerator::ExportProperties (or that code made commonly available in editor mode)
-
 	// Emit class variables
 	for (TFieldIterator<UProperty> It(SourceClass, EFieldIteratorFlags::ExcludeSuper); It; ++It)
 	{
@@ -372,18 +542,18 @@ void FKismetCppBackend::EmitClassProperties(FStringOutputDevice& Target, UClass*
 
 		Emit(Header, TEXT("\n\tUPROPERTY("));
 		{
-			TArray<FString> Tags = FEmitProperyHelper::ProperyFlagsToTags(Property->PropertyFlags);
-			Tags.Emplace(FEmitProperyHelper::HandleRepNotifyFunc(Property));
-			Tags.Emplace(FEmitProperyHelper::HandleMetaData(Property));
+			TArray<FString> Tags = FEmitHelper::ProperyFlagsToTags(Property->PropertyFlags);
+			Tags.Emplace(FEmitHelper::HandleRepNotifyFunc(Property));
+			Tags.Emplace(FEmitHelper::HandleMetaData(Property));
 			Tags.Remove(FString());
 
 			FString AllTags;
-			FEmitProperyHelper::ArrayToString(Tags, AllTags);
+			FEmitHelper::ArrayToString(Tags, AllTags, TEXT(", "));
 			Emit(Header, *AllTags);
 		}
 		Emit(Header, TEXT(")\n"));
 		Emit(Header, TEXT("\t"));
-		Property->ExportCppDeclaration(Target, EExportedDeclaration::Member);
+		Property->ExportCppDeclaration(Target, EExportedDeclaration::Member, NULL, EPropertyExportCPPFlags::CPPF_CustomTypeName);
 		Emit(Header, TEXT(";\n"));
 	}
 }
@@ -396,6 +566,45 @@ void FKismetCppBackend::GenerateCodeFromClass(UClass* SourceClass, TIndirectArra
 
 	Emit(Header, TEXT("#pragma once\n\n"));
 	Emit(Header, *FString::Printf(TEXT("#include \"%s.generated.h\"\n\n"), *SourceClass->GetName()));
+
+	// MC DELEGATE DECLARATION
+	{
+		auto DelegateDeclarations = FEmitHelper::EmitMulticastDelegateDeclarations(SourceClass);
+		FString AllDeclarations;
+		FEmitHelper::ArrayToString(DelegateDeclarations, AllDeclarations, TEXT(";\n"));
+		if (DelegateDeclarations.Num())
+		{
+			Emit(Header, *AllDeclarations);
+			Emit(Header, TEXT(";\n"));
+		}
+	}
+
+	// GATHER ALL SC DELEGATES
+	{
+		TArray<UMulticastDelegateProperty*> Delegates;
+		for (TFieldIterator<UMulticastDelegateProperty> It(SourceClass, EFieldIteratorFlags::ExcludeSuper); It; ++It)
+		{
+			Delegates.Add(*It);
+		}
+
+		for (auto& FuncContext : Functions)
+		{
+			for (TFieldIterator<UMulticastDelegateProperty> It(FuncContext.Function, EFieldIteratorFlags::ExcludeSuper); It; ++It)
+			{
+				Delegates.Add(*It);
+			}
+		}
+
+		auto DelegateDeclarations = FEmitHelper::EmitSinglecastDelegateDeclarations(Delegates);
+		FString AllDeclarations;
+		FEmitHelper::ArrayToString(DelegateDeclarations, AllDeclarations, TEXT(";\n"));
+		if (DelegateDeclarations.Num())
+		{
+			Emit(Header, *AllDeclarations);
+			Emit(Header, TEXT(";\n"));
+		}
+	}
+
 	Emit(Header, TEXT("UCLASS(Blueprintable)\n"));
 
 	Emit(Header,
@@ -440,7 +649,12 @@ void FKismetCppBackend::GenerateCodeFromClass(UClass* SourceClass, TIndirectArra
 
 void FKismetCppBackend::EmitCallDelegateStatment(FKismetFunctionContext& FunctionContext, FBlueprintCompiledStatement& Statement)
 {
-	Emit(Body, *FString::Printf(TEXT("\t\t\t%s.Broadcast("), *TermToText(Statement.FunctionContext)));
+	check(Statement.FunctionContext && Statement.FunctionContext->AssociatedVarProperty);
+	FSafeContextScopedEmmitter SafeContextScope(Body, Statement.FunctionContext->Context, *this, TEXT("\t\t\t"));
+
+	Emit(Body, TEXT("\t\t\t"));
+	Emit(Body, *SafeContextScope.GetAdditionalIndent());
+	Emit(Body, *FString::Printf(TEXT("%s.Broadcast("), *TermToText(Statement.FunctionContext)));
 	int32 NumParams = 0;
 	for (TFieldIterator<UProperty> PropIt(Statement.FunctionToCall); PropIt && (PropIt->PropertyFlags & CPF_Parm); ++PropIt)
 	{
@@ -500,7 +714,14 @@ void FKismetCppBackend::EmitCallDelegateStatment(FKismetFunctionContext& Functio
 
 void FKismetCppBackend::EmitCallStatment(FKismetFunctionContext& FunctionContext, FBlueprintCompiledStatement& Statement)
 {
+	const bool bCallOnDifferentObject = ((Statement.FunctionContext != NULL) && (Statement.FunctionContext->Name != Schema->PSC_Self));
+	const bool bStaticCall = Statement.FunctionToCall->HasAnyFunctionFlags(FUNC_Static);
+
+	const bool bUseSafeContext = bCallOnDifferentObject && !bStaticCall;
+	FSafeContextScopedEmmitter SafeContextScope(Body, bUseSafeContext ? Statement.FunctionContext : NULL, *this, TEXT("\t\t\t"));
+
 	Emit(Body, TEXT("\t\t\t"));
+	Emit(Body, *SafeContextScope.GetAdditionalIndent());
 
 	// Handle the return value of the function being called
 	UProperty* FuncToCallReturnProperty = Statement.FunctionToCall->GetReturnProperty();
@@ -514,9 +735,9 @@ void FKismetCppBackend::EmitCallStatment(FKismetFunctionContext& FunctionContext
 	}
 
 	// Emit object to call the method on
-	if ((Statement.FunctionContext != NULL) && (Statement.FunctionContext->Name != Schema->PSC_Self)) //@TODO: Badness, could be a self reference wired to another instance!
+	if (bCallOnDifferentObject) //@TODO: Badness, could be a self reference wired to another instance!
 	{
-		if (Statement.FunctionToCall->HasAnyFunctionFlags(FUNC_Static))
+		if (bStaticCall)
 		{
 			const bool bIsCustomThunk = Statement.FunctionToCall->HasMetaData(TEXT("CustomStructureParam")) || Statement.FunctionToCall->HasMetaData(TEXT("ArrayParm"));
 			FString FullFunctionNamePrefix = bIsCustomThunk
@@ -604,14 +825,14 @@ void FKismetCppBackend::EmitCallStatment(FKismetFunctionContext& FunctionContext
 
 void FKismetCppBackend::EmitAssignmentStatment(FKismetFunctionContext& FunctionContext, FBlueprintCompiledStatement& Statement)
 {
-	FString DestinationExpression;
-	DestinationExpression = TermToText(Statement.LHS);
-
-	FString SourceExpression;
-	SourceExpression = TermToText(Statement.RHS[0], Statement.LHS->AssociatedVarProperty);
+	FString DestinationExpression = TermToText(Statement.LHS);
+	FString SourceExpression = TermToText(Statement.RHS[0], Statement.LHS->AssociatedVarProperty);
 
+	FSafeContextScopedEmmitter SafeContextScope(Body, Statement.LHS->Context, *this, TEXT("\t\t\t"));
 	// Emit the assignment statement
-	Emit(Body, *FString::Printf(TEXT("\t\t\t%s = %s;\n"), *DestinationExpression, *SourceExpression));
+	Emit(Body, TEXT("\t\t\t"));
+	Emit(Body, *SafeContextScope.GetAdditionalIndent());
+	Emit(Body, *FString::Printf(TEXT("%s = %s;\n"), *DestinationExpression, *SourceExpression));
 }
 
 void FKismetCppBackend::EmitCastObjToInterfaceStatement(FKismetFunctionContext& FunctionContext, FBlueprintCompiledStatement& Statement)
@@ -620,7 +841,7 @@ void FKismetCppBackend::EmitCastObjToInterfaceStatement(FKismetFunctionContext&
 	FString ObjectValue = TermToText(Statement.RHS[1], (UProperty*)(GetDefault<UObjectProperty>()));
 	FString InterfaceValue = TermToText(Statement.LHS, (UProperty*)(GetDefault<UObjectProperty>()));
 
-	Emit(Body, *FString::Printf(TEXT("\t\t\tif ( %s && %s->GetClass()->ImplementsInterface(%s) )\n"), *ObjectValue, *ObjectValue, *InterfaceClass));
+	Emit(Body, *FString::Printf(TEXT("\t\t\tif ( IsValid(%s) && %s->GetClass()->ImplementsInterface(%s) )\n"), *ObjectValue, *ObjectValue, *InterfaceClass));
 	Emit(Body, *FString::Printf(TEXT("\t\t\t{\n")));
 	Emit(Body, *FString::Printf(TEXT("\t\t\t\t%s.SetObject(%s);\n"), *InterfaceValue, *ObjectValue));
 	Emit(Body, *FString::Printf(TEXT("\t\t\t\tvoid* IAddress = %s->GetInterfaceAddress(%s);\n"), *ObjectValue, *InterfaceClass));
@@ -682,42 +903,60 @@ void FKismetCppBackend::EmitObjectToBoolStatement(FKismetFunctionContext& Functi
 
 void FKismetCppBackend::EmitAddMulticastDelegateStatement(FKismetFunctionContext& FunctionContext, FBlueprintCompiledStatement& Statement)
 {
+	check(Statement.LHS && Statement.LHS->AssociatedVarProperty);
+	FSafeContextScopedEmmitter SafeContextScope(Body, Statement.LHS->Context, *this, TEXT("\t\t\t"));
+
+	Emit(Body, TEXT("\t\t\t"));
+	Emit(Body, *SafeContextScope.GetAdditionalIndent());
+
 	const FString Delegate = TermToText(Statement.LHS);
 	const FString DelegateToAdd = TermToText(Statement.RHS[0]);
 
-	Emit(Body, *FString::Printf(TEXT("\t\t\t%s->Add(%s);\n"), *Delegate, *DelegateToAdd));
+	Emit(Body, *FString::Printf(TEXT("%s.Add(%s);\n"), *Delegate, *DelegateToAdd));
 }
 
 void FKismetCppBackend::EmitRemoveMulticastDelegateStatement(FKismetFunctionContext& FunctionContext, FBlueprintCompiledStatement& Statement)
 {
+	check(Statement.LHS && Statement.LHS->AssociatedVarProperty);
+	FSafeContextScopedEmmitter SafeContextScope(Body, Statement.LHS->Context, *this, TEXT("\t\t\t"));
+
+	Emit(Body, TEXT("\t\t\t"));
+	Emit(Body, *SafeContextScope.GetAdditionalIndent());
+
 	const FString Delegate = TermToText(Statement.LHS);
 	const FString DelegateToAdd = TermToText(Statement.RHS[0]);
 
-	Emit(Body, *FString::Printf(TEXT("\t\t\t%s->Remove(%s);\n"), *Delegate, *DelegateToAdd));
+	Emit(Body, *FString::Printf(TEXT("%s.Remove(%s);\n"), *Delegate, *DelegateToAdd));
 }
 
 void FKismetCppBackend::EmitBindDelegateStatement(FKismetFunctionContext& FunctionContext, FBlueprintCompiledStatement& Statement)
 {
 	check(2 == Statement.RHS.Num());
+	check(Statement.LHS);
+	FSafeContextScopedEmmitter SafeContextScope(Body, Statement.LHS->Context, *this, TEXT("\t\t\t"));
+
+	Emit(Body, TEXT("\t\t\t"));
+	Emit(Body, *SafeContextScope.GetAdditionalIndent());
+
+
 	const FString Delegate = TermToText(Statement.LHS);
 	const FString NameTerm = TermToText(Statement.RHS[0], GetDefault<UNameProperty>());
 	const FString ObjectTerm = TermToText(Statement.RHS[1], GetDefault<UObjectProperty>());
 
-	Emit(Body, 
-		*FString::Printf(
-			TEXT("\t\t\tif(FScriptDelegate* __Delegate = %s)\n\t\t\t{\n\t\t\t\t__Delegate->BindUFunction(%s,%s);\n\t\t\t|\n"),
-			*Delegate, 
-			*ObjectTerm,
-			*NameTerm
-		)
-	);
+	Emit(Body, *FString::Printf(TEXT("%s.BindUFunction(%s,%s);\n"), *Delegate, *ObjectTerm, *NameTerm));
 }
 
 void FKismetCppBackend::EmitClearMulticastDelegateStatement(FKismetFunctionContext& FunctionContext, FBlueprintCompiledStatement& Statement)
 {
+	check(Statement.LHS);
+	FSafeContextScopedEmmitter SafeContextScope(Body, Statement.LHS->Context, *this, TEXT("\t\t\t"));
+
+	Emit(Body, TEXT("\t\t\t"));
+	Emit(Body, *SafeContextScope.GetAdditionalIndent());
+
 	const FString Delegate = TermToText(Statement.LHS);
 
-	Emit(Body, *FString::Printf(TEXT("\t\t\t%s->Clear();\n"), *Delegate));
+	Emit(Body, *FString::Printf(TEXT("%s.Clear();\n"), *Delegate));
 }
 
 void FKismetCppBackend::EmitCreateArrayStatement(FKismetFunctionContext& FunctionContext, FBlueprintCompiledStatement& Statement)
@@ -800,7 +1039,7 @@ void FKismetCppBackend::DeclareLocalVariables(FKismetFunctionContext& FunctionCo
 		UProperty* LocalVariable = LocalVariables[i];
 
 		Emit(Body, TEXT("\t"));
-		LocalVariable->ExportCppDeclaration(Body, EExportedDeclaration::Local);
+		LocalVariable->ExportCppDeclaration(Body, EExportedDeclaration::Local, NULL, EPropertyExportCPPFlags::CPPF_CustomTypeName);
 		Emit(Body, TEXT(";\n"));
 	}
 
@@ -835,93 +1074,103 @@ void FKismetCppBackend::CloseStateSwitch(FKismetFunctionContext& FunctionContext
 
 void FKismetCppBackend::ConstructFunction(FKismetFunctionContext& FunctionContext, bool bGenerateStubOnly)
 {
-	FString FunctionName;
-	UFunction* Function = FunctionContext.Function;
-	UClass* Class = FunctionContext.NewClass;
+	if (FunctionContext.IsDelegateSignature())
+	{
+		return;
+	}
 
-	Function->GetName(FunctionName);
+	UFunction* Function = FunctionContext.Function;
 
-	// Split the function property list into arguments, a return value (if any), and local variable declarations
-	TArray<UProperty*> ArgumentList;
 	UProperty* ReturnValue = NULL;
 	TArray<UProperty*> LocalVariables;
-	for (TFieldIterator<UProperty> It(Function); It; ++It)
+	
 	{
-		UProperty* Property = *It;
-		if (Property->HasAnyPropertyFlags(CPF_Parm))
+		FString FunctionName;
+		Function->GetName(FunctionName);
+
+		TArray<UProperty*> ArgumentList;
+
+		// Split the function property list into arguments, a return value (if any), and local variable declarations
+		for (TFieldIterator<UProperty> It(Function); It; ++It)
 		{
-			if (Property->HasAnyPropertyFlags(CPF_ReturnParm))
+			UProperty* Property = *It;
+			if (Property->HasAnyPropertyFlags(CPF_Parm))
 			{
-				if (ReturnValue == NULL)
+				if (Property->HasAnyPropertyFlags(CPF_ReturnParm))
 				{
-					ReturnValue = Property;
-					LocalVariables.Add(Property);
+					if (ReturnValue == NULL)
+					{
+						ReturnValue = Property;
+						LocalVariables.Add(Property);
+					}
+					else
+					{
+						MessageLog.Error(*FString::Printf(TEXT("Function %s from graph @@ has more than one return value (named %s and %s)"),
+							*FunctionName, *ReturnValue->GetName(), *Property->GetName()), FunctionContext.SourceGraph);
+					}
 				}
 				else
 				{
-					MessageLog.Error(*FString::Printf(TEXT("Function %s from graph @@ has more than one return value (named %s and %s)"),
-						*FunctionName, *ReturnValue->GetName(), *Property->GetName()), FunctionContext.SourceGraph);
+					ArgumentList.Add(Property);
 				}
 			}
 			else
 			{
-				ArgumentList.Add(Property);
-			}	
-		}
-		else
-		{
-			LocalVariables.Add(Property);
+				LocalVariables.Add(Property);
+			}
 		}
-	}
 
-	// Emit the declaration
-	FString ReturnValueString;
-	FString ReturnType;
-	if (ReturnValue != NULL)
-	{
-		ReturnType = ReturnValue->GetCPPType(NULL);
-		ReturnValueString = FString(TEXT(" ")) + ReturnValue->GetName();
-	}
-	else
-	{
-		ReturnType = TEXT("void");
-		ReturnValueString = TEXT("");
-	}
-
-	//@TODO: Make the header+body export more uniform
-	{
-		FString Start = FString::Printf(TEXT("%s %s%s%s("), *ReturnType, TEXT("%s"), TEXT("%s"), *FunctionName);
+		// Emit the declaration
+		const FString ReturnType = ReturnValue ? ReturnValue->GetCPPType(NULL, EPropertyExportCPPFlags::CPPF_CustomTypeName) : TEXT("void");
 
-		Emit(Header, TEXT("\t"));
-		Emit(Header, *FString::Printf(*Start, TEXT(""), TEXT("")));
-		Emit(Body, *FString::Printf(*Start, *CppClassName, TEXT("::")));
-
-		for (int32 i = 0; i < ArgumentList.Num(); ++i)
+		//@TODO: Make the header+body export more uniform
 		{
-			UProperty* ArgProperty = ArgumentList[i];
+			const FString Start = FString::Printf(TEXT("%s %s%s%s("), *ReturnType, TEXT("%s"), TEXT("%s"), *FunctionName);
 
-			if (i > 0)
+			const bool bUserConstructionScript = (FunctionName == TEXT("UserConstructionScript"));
+			if (!bUserConstructionScript)
 			{
-				Emit(Header, TEXT(", "));
-				Emit(Body, TEXT(", "));
+				Emit(Header, *FString::Printf(TEXT("\t%s\n"), *FEmitHelper::EmitUFuntion(Function)));
 			}
+			Emit(Header, TEXT("\t"));
+			Emit(Header, *FString::Printf(*Start, TEXT(""), TEXT("")));
+			Emit(Body, *FString::Printf(*Start, *CppClassName, TEXT("::")));
 
-			if (ArgProperty->HasAnyPropertyFlags(CPF_OutParm))
+			for (int32 i = 0; i < ArgumentList.Num(); ++i)
 			{
-				Emit(Header, TEXT("/*out*/ "));
-				Emit(Body, TEXT("/*out*/ "));
+				UProperty* ArgProperty = ArgumentList[i];
+
+				if (i > 0)
+				{
+					Emit(Header, TEXT(", "));
+					Emit(Body, TEXT(", "));
+				}
+
+				if (ArgProperty->HasAnyPropertyFlags(CPF_OutParm))
+				{
+					Emit(Header, TEXT("/*out*/ "));
+					Emit(Body, TEXT("/*out*/ "));
+				}
+
+				ArgProperty->ExportCppDeclaration(Header, EExportedDeclaration::Parameter, NULL, EPropertyExportCPPFlags::CPPF_CustomTypeName);
+				ArgProperty->ExportCppDeclaration(Body, EExportedDeclaration::Parameter, NULL, EPropertyExportCPPFlags::CPPF_CustomTypeName);
 			}
 
-			ArgProperty->ExportCppDeclaration(Header, EExportedDeclaration::Parameter);
-			ArgProperty->ExportCppDeclaration(Body,   EExportedDeclaration::Parameter);
+			Emit(Header, TEXT(")"));
+			if (bUserConstructionScript)
+			{
+				Emit(Header, TEXT(" override"));
+			}
+			Emit(Header, TEXT(";\n"));
+			
+			Emit(Body, TEXT(")\n"));
 		}
 
-		Emit(Header, TEXT(");\n"));
-		Emit(Body, TEXT(")\n"));
+		// Start the body of the implementation
+		Emit(Body, TEXT("{\n"));
 	}
 
-	// Start the body of the implementation
-	Emit(Body, TEXT("{\n"));
+	const FString ReturnValueString = ReturnValue ? (FString(TEXT(" ")) + ReturnValue->GetName()) : TEXT("");
 
 	if( !bGenerateStubOnly )
 	{
diff --git a/Engine/Source/Editor/KismetCompiler/Private/KismetCompilerMisc.cpp b/Engine/Source/Editor/KismetCompiler/Private/KismetCompilerMisc.cpp
--- a/Engine/Source/Editor/KismetCompiler/Private/KismetCompilerMisc.cpp
+++ b/Engine/Source/Editor/KismetCompiler/Private/KismetCompilerMisc.cpp
@@ -1190,12 +1190,12 @@ FBPTerminal* FKismetFunctionContext::CreateLocalTerminalFromPinAutoChooseScope(U
 
 	const bool bRegularBlueprint = Blueprint && (Blueprint->GetClass() == UBlueprint::StaticClass());
 	static FBoolConfigValueHelper UseLocalGraphVariables(TEXT("Kismet"), TEXT("bUseLocalGraphVariables"), GEngineIni);
-	if (bSharedTerm && bRegularBlueprint && UseLocalGraphVariables)
+	const bool OutputPin = EEdGraphPinDirection::EGPD_Output == Net->Direction;
+	if (bSharedTerm && bRegularBlueprint && UseLocalGraphVariables && OutputPin)
 	{
 		SCOPE_CYCLE_COUNTER(EKismetCompilerStats_ChooseTerminalScope);
 
 		// Pin's connections are checked, to tell if created terminal is shared, or if it could be a local variable.
-		ensure(EEdGraphPinDirection::EGPD_Output == Net->Direction);
 		bSharedTerm = FEventGraphUtils::PinRepresentsSharedTerminal(*Net);
 	}
 	FBPTerminal* Term = new (bSharedTerm ? EventGraphLocals : Locals) FBPTerminal();
diff --git a/Engine/Source/Runtime/CoreUObject/Private/UObject/Property.cpp b/Engine/Source/Runtime/CoreUObject/Private/UObject/Property.cpp
--- a/Engine/Source/Runtime/CoreUObject/Private/UObject/Property.cpp
+++ b/Engine/Source/Runtime/CoreUObject/Private/UObject/Property.cpp
@@ -476,14 +476,15 @@ FString UProperty::GetCPPMacroType( FString& ExtendedTypeText ) const
 	return TEXT("PROPERTY");
 }
 
-void UProperty::ExportCppDeclaration( FOutputDevice& Out, EExportedDeclaration::Type DeclarationType, const TCHAR* ArrayDimOverride ) const
+void UProperty::ExportCppDeclaration(FOutputDevice& Out, EExportedDeclaration::Type DeclarationType, const TCHAR* ArrayDimOverride, uint32 AdditionalExportCPPFlags) const
 {
-	const bool bIsParameter     = DeclarationType == EExportedDeclaration::Parameter;
+	const bool bIsParameter = (DeclarationType == EExportedDeclaration::Parameter) || (DeclarationType == EExportedDeclaration::MacroParameter);
 	const bool bIsInterfaceProp = dynamic_cast<const UInterfaceProperty*>(this) != nullptr;
 
 	// export the property type text (e.g. FString; int32; TArray, etc.)
 	FString ExtendedTypeText;
-	FString TypeText = GetCPPType(&ExtendedTypeText, bIsParameter ? CPPF_ArgumentOrReturnValue : 0);
+	const uint32 ExportCPPFlags = AdditionalExportCPPFlags | (bIsParameter ? CPPF_ArgumentOrReturnValue : 0);
+	FString TypeText = GetCPPType(&ExtendedTypeText, ExportCPPFlags);
 
 	if (!dynamic_cast<const UBoolProperty*>(this)) // can't have const bitfields because then we cannot determine their offset and mask from the compiler
 	{
@@ -507,6 +508,12 @@ void UProperty::ExportCppDeclaration( FOutputDevice& Out, EExportedDeclaration::
 		}
 	}
 
+	FString NameCpp = GetNameCPP();
+	if (DeclarationType == EExportedDeclaration::MacroParameter)
+	{
+		NameCpp = FString(TEXT(", ")) + NameCpp;
+	}
+
 	TCHAR ArrayStr[MAX_SPRINTF]=TEXT("");
 	if( ArrayDim != 1 )
 	{
@@ -526,19 +533,19 @@ void UProperty::ExportCppDeclaration( FOutputDevice& Out, EExportedDeclaration::
 		if( ArrayDim==1 && DeclarationType == EExportedDeclaration::Member )
 		{
 			// export as a uint32 member....bad to hardcode, but this is a special case that won't be used anywhere else
-			Out.Logf( TEXT("%s%s %s%s:1"), *TypeText, *ExtendedTypeText, *GetNameCPP(), ArrayStr );
+			Out.Logf(TEXT("%s%s %s%s:1"), *TypeText, *ExtendedTypeText, *NameCpp, ArrayStr);
 		}
 
 		//@todo we currently can't have out bools.. so this isn't really necessary, but eventually out bools may be supported, so leave here for now
 		else if( bIsParameter && HasAnyPropertyFlags(CPF_OutParm) )
 		{
 			// export as a reference
-			Out.Logf( TEXT("%s%s%s %s%s"), *TypeText, *ExtendedTypeText, TEXT("&"), *GetNameCPP(), ArrayStr );
+			Out.Logf(TEXT("%s%s%s %s%s"), *TypeText, *ExtendedTypeText, TEXT("&"), *NameCpp, ArrayStr);
 		}
 
 		else
 		{
-			Out.Logf( TEXT("%s%s %s%s"), *TypeText, *ExtendedTypeText, *GetNameCPP(), ArrayStr );
+			Out.Logf(TEXT("%s%s %s%s"), *TypeText, *ExtendedTypeText, *NameCpp, ArrayStr);
 		}
 	}
 	else 
@@ -550,7 +557,7 @@ void UProperty::ExportCppDeclaration( FOutputDevice& Out, EExportedDeclaration::
 				// export as a pointer
 				//Out.Logf( TEXT("%s%s* %s"), *TypeText, *ExtendedTypeText, *GetNameCPP() );
 				// don't export as a pointer
-				Out.Logf( TEXT("%s%s %s%s"), *TypeText, *ExtendedTypeText, *GetNameCPP(), ArrayStr );
+				Out.Logf(TEXT("%s%s %s%s"), *TypeText, *ExtendedTypeText, *NameCpp, ArrayStr);
 			}
 			else
 			{
@@ -561,7 +568,7 @@ void UProperty::ExportCppDeclaration( FOutputDevice& Out, EExportedDeclaration::
 						!HasAnyPropertyFlags(CPF_OutParm|CPF_ConstParm) ? TEXT("const ") : TEXT(""),
 						*TypeText, *ExtendedTypeText,
 						TEXT("&"),
-						*GetNameCPP());
+						*NameCpp);
 				}
 				else
 				{
@@ -571,13 +578,13 @@ void UProperty::ExportCppDeclaration( FOutputDevice& Out, EExportedDeclaration::
 					{
 						ModifierString[0] = TEXT('&');
 					}
-					Out.Logf( TEXT("%s%s%s %s%s"), *TypeText, *ExtendedTypeText, ModifierString, *GetNameCPP(), ArrayStr );
+					Out.Logf(TEXT("%s%s%s %s%s"), *TypeText, *ExtendedTypeText, ModifierString, *NameCpp, ArrayStr);
 				}
 			}
 		}
 		else
 		{
-			Out.Logf( TEXT("%s%s %s%s"), *TypeText, *ExtendedTypeText, *GetNameCPP(), ArrayStr );
+			Out.Logf(TEXT("%s%s %s%s"), *TypeText, *ExtendedTypeText, *NameCpp, ArrayStr);
 		}
 	}
 }
diff --git a/Engine/Source/Runtime/CoreUObject/Private/UObject/PropertyDelegate.cpp b/Engine/Source/Runtime/CoreUObject/Private/UObject/PropertyDelegate.cpp
--- a/Engine/Source/Runtime/CoreUObject/Private/UObject/PropertyDelegate.cpp
+++ b/Engine/Source/Runtime/CoreUObject/Private/UObject/PropertyDelegate.cpp
@@ -73,9 +73,12 @@ bool UDelegateProperty::NetSerializeItem( FArchive& Ar, UPackageMap* Map, void*
 
 FString UDelegateProperty::GetCPPType( FString* ExtendedTypeText/*=NULL*/, uint32 CPPExportFlags/*=0*/ ) const
 {
-	const FString UnmangledFunctionName = SignatureFunction->GetName().LeftChop( FString( HEADER_GENERATED_DELEGATE_SIGNATURE_SUFFIX ).Len() );
-	const FString DelegateName( FString( TEXT( "F" ) ) + UnmangledFunctionName );
-	return DelegateName;
+	FString UnmangledFunctionName = SignatureFunction->GetName().LeftChop( FString( HEADER_GENERATED_DELEGATE_SIGNATURE_SUFFIX ).Len() );
+	if (0 != (CPPExportFlags & EPropertyExportCPPFlags::CPPF_CustomTypeName))
+	{
+		UnmangledFunctionName += TEXT("__SinglecastDelegate");
+	}
+	return FString(TEXT("F")) + UnmangledFunctionName;
 }
 
 
diff --git a/Engine/Source/Runtime/CoreUObject/Private/UObject/PropertyMulticastDelegate.cpp b/Engine/Source/Runtime/CoreUObject/Private/UObject/PropertyMulticastDelegate.cpp
--- a/Engine/Source/Runtime/CoreUObject/Private/UObject/PropertyMulticastDelegate.cpp
+++ b/Engine/Source/Runtime/CoreUObject/Private/UObject/PropertyMulticastDelegate.cpp
@@ -112,9 +112,12 @@ bool UMulticastDelegateProperty::NetSerializeItem( FArchive& Ar, UPackageMap* Ma
 
 FString UMulticastDelegateProperty::GetCPPType( FString* ExtendedTypeText/*=NULL*/, uint32 CPPExportFlags/*=0*/ ) const
 {
-	const FString UnmangledFunctionName = SignatureFunction->GetName().LeftChop( FString( HEADER_GENERATED_DELEGATE_SIGNATURE_SUFFIX ).Len() );
-	const FString DelegateName( FString( TEXT( "F" ) ) + UnmangledFunctionName );
-	return DelegateName;
+	FString UnmangledFunctionName = SignatureFunction->GetName().LeftChop( FString( HEADER_GENERATED_DELEGATE_SIGNATURE_SUFFIX ).Len() );
+	if (0 != (CPPExportFlags & EPropertyExportCPPFlags::CPPF_CustomTypeName))
+	{
+		UnmangledFunctionName += TEXT("__MulticastDelegate");
+	}
+	return FString(TEXT("F")) + UnmangledFunctionName;
 }
 
 
diff --git a/Engine/Source/Runtime/CoreUObject/Public/UObject/UnrealType.h b/Engine/Source/Runtime/CoreUObject/Public/UObject/UnrealType.h
--- a/Engine/Source/Runtime/CoreUObject/Public/UObject/UnrealType.h
+++ b/Engine/Source/Runtime/CoreUObject/Public/UObject/UnrealType.h
@@ -112,6 +112,8 @@ enum EPropertyExportCPPFlags
 	CPPF_ArgumentOrReturnValue		=	0x00000002,
 	/** Indicates thet we are exporting this property's CPP text for C++ definition of a function. */
 	CPPF_Implementation = 0x00000004,
+	/** Indicates thet we are exporting this property's CPP text with an custom type name */
+	CPPF_CustomTypeName				=	0x00000008,
 };
 
 namespace EExportedDeclaration
@@ -120,7 +122,9 @@ namespace EExportedDeclaration
 	{
 		Local,
 		Member,
-		Parameter
+		Parameter,
+		/** Type and mane are separated by comma */
+		MacroParameter, 
 	};
 }
 
@@ -178,7 +182,7 @@ public:
 											FOutputDevice* Warn, TArray<struct FDefinedProperty>& DefinedProperties );
 
 	// UHT interface
-	void ExportCppDeclaration( FOutputDevice& Out, EExportedDeclaration::Type DeclarationType, const TCHAR* ArrayDimOverride = NULL ) const;
+	void ExportCppDeclaration(FOutputDevice& Out, EExportedDeclaration::Type DeclarationType, const TCHAR* ArrayDimOverride = NULL, uint32 AdditionalExportCPPFlags = 0) const;
 	virtual FString GetCPPMacroType( FString& ExtendedTypeText ) const;
 	virtual bool PassCPPArgsByRef() const { return false; }
 
@@ -738,7 +742,7 @@ public:
 	 */
 	UProperty* GetOwnerProperty()
 	{
-		UProperty* Result = this;
+		UProperty* Result=this;
 		for (UProperty* PropBase = dynamic_cast<UProperty*>(GetOuter()); PropBase; PropBase = dynamic_cast<UProperty*>(PropBase->GetOuter()))
 		{
 			Result = PropBase;
