BP: FKismetCppBackend - ongoing task.
UProperty::ExportCppDeclaration wont generate bitfield for "bool".

diff --git a/Engine/Source/Editor/KismetCompiler/Private/KismetCompilerCppBackend.cpp b/Engine/Source/Editor/KismetCompiler/Private/KismetCompilerCppBackend.cpp
--- a/Engine/Source/Editor/KismetCompiler/Private/KismetCompilerCppBackend.cpp
+++ b/Engine/Source/Editor/KismetCompiler/Private/KismetCompilerCppBackend.cpp
@@ -113,12 +113,12 @@ FString FKismetCppBackend::TermToText(const FBPTerminal* Term, const UProperty*
 		else if (auto BoolProperty = Cast<const UBoolProperty>(CoerceProperty))
 		{
 			bool bValue = Term->Name.ToBool();
-			return *(FName::GetEntry(bValue ? NAME_TRUE : NAME_FALSE)->GetPlainNameString());
+			return bValue ? TEXT("true") : TEXT("false");
 		}
 		else if (auto NameProperty = Cast<const UNameProperty>(CoerceProperty))
 		{
 			FName LiteralName(*(Term->Name));
-			return FString::Printf(TEXT("FName(TEXT(\"%s\")"), *(LiteralName.ToString()));
+			return FString::Printf(TEXT("FName(TEXT(\"%s\"))"), *(LiteralName.ToString()));
 		}
 		else if (auto StructProperty = Cast<const UStructProperty>(CoerceProperty))
 		{
@@ -309,7 +309,7 @@ struct FEmitHelper
 		return FString();
 	}
 
-	static FString HandleMetaData(const UField* Field)
+	static FString HandleMetaData(const UField* Field, bool AddCategory = true)
 	{
 		FString MetaDataStr;
 
@@ -319,9 +319,9 @@ struct FEmitHelper
 		const UMetaData* MetaData = Package->GetMetaData();
 		check(MetaData);
 		const TMap<FName, FString>* ValuesMap = MetaData->ObjectMetaDataMap.Find(Field);
+		TArray<FString> MetaDataStrings;
 		if (ValuesMap && ValuesMap->Num())
 		{
-			TArray<FString> MetaDataStrings;
 			for (auto& Pair : *ValuesMap)
 			{
 				if (!Pair.Value.IsEmpty())
@@ -333,13 +333,17 @@ struct FEmitHelper
 					MetaDataStrings.Emplace(Pair.Key.ToString());
 				}
 			}
-			MetaDataStrings.Remove(FString());
-			if (MetaDataStrings.Num())
-			{
-				MetaDataStr += TEXT("meta=(");
-				ArrayToString(MetaDataStrings, MetaDataStr, TEXT(", "));
-				MetaDataStr += TEXT(")");
-			}
+		}
+		if (AddCategory && (!ValuesMap || !ValuesMap->Find(TEXT("Category"))))
+		{
+			MetaDataStrings.Emplace(TEXT("Category"));
+		}
+		MetaDataStrings.Remove(FString());
+		if (MetaDataStrings.Num())
+		{
+			MetaDataStr += TEXT("meta=(");
+			ArrayToString(MetaDataStrings, MetaDataStr, TEXT(", "));
+			MetaDataStr += TEXT(")");
 		}
 		return MetaDataStr;
 	}
@@ -455,7 +459,8 @@ struct FEmitHelper
 	static FString EmitUFuntion(UFunction* Function)
 	{
 		TArray<FString> Tags = FEmitHelper::FunctionFlagsToTags(Function->FunctionFlags);
-		Tags.Emplace(FEmitHelper::HandleMetaData(Function));
+		const bool bMustHaveCategory = (Function->FunctionFlags & (FUNC_BlueprintCallable | FUNC_BlueprintPure)) != 0;
+		Tags.Emplace(FEmitHelper::HandleMetaData(Function, bMustHaveCategory));
 		Tags.Remove(FString());
 
 		FString AllTags;
@@ -495,7 +500,7 @@ struct FEmitHelper
 		return ParameterNum;
 	}
 
-	static TArray<FString> EmitSinglecastDelegateDeclarations(const TArray<UMulticastDelegateProperty*>& Delegates)
+	static TArray<FString> EmitSinglecastDelegateDeclarations(const TArray<UDelegateProperty*>& Delegates)
 	{
 		TArray<FString> Results;
 		for (auto It : Delegates)
@@ -544,7 +549,7 @@ void FKismetCppBackend::EmitClassProperties(FStringOutputDevice& Target, UClass*
 		{
 			TArray<FString> Tags = FEmitHelper::ProperyFlagsToTags(Property->PropertyFlags);
 			Tags.Emplace(FEmitHelper::HandleRepNotifyFunc(Property));
-			Tags.Emplace(FEmitHelper::HandleMetaData(Property));
+			Tags.Emplace(FEmitHelper::HandleMetaData(Property, false));
 			Tags.Remove(FString());
 
 			FString AllTags;
@@ -581,15 +586,15 @@ void FKismetCppBackend::GenerateCodeFromClass(UClass* SourceClass, TIndirectArra
 
 	// GATHER ALL SC DELEGATES
 	{
-		TArray<UMulticastDelegateProperty*> Delegates;
-		for (TFieldIterator<UMulticastDelegateProperty> It(SourceClass, EFieldIteratorFlags::ExcludeSuper); It; ++It)
+		TArray<UDelegateProperty*> Delegates;
+		for (TFieldIterator<UDelegateProperty> It(SourceClass, EFieldIteratorFlags::ExcludeSuper); It; ++It)
 		{
 			Delegates.Add(*It);
 		}
 
 		for (auto& FuncContext : Functions)
 		{
-			for (TFieldIterator<UMulticastDelegateProperty> It(FuncContext.Function, EFieldIteratorFlags::ExcludeSuper); It; ++It)
+			for (TFieldIterator<UDelegateProperty> It(FuncContext.Function, EFieldIteratorFlags::ExcludeSuper); It; ++It)
 			{
 				Delegates.Add(*It);
 			}
@@ -629,6 +634,7 @@ void FKismetCppBackend::GenerateCodeFromClass(UClass* SourceClass, TIndirectArra
 		Emit(Header, TEXT("\n"));
 	}
 
+	Emit(Body, *FString::Printf(TEXT("#include \"%s.h\"\n"), FApp::GetGameName()));
 	Emit(Body, TEXT("#include \"GeneratedCodeHelpers.h\"\n\n"));
 
 	//constructor
diff --git a/Engine/Source/Runtime/CoreUObject/Private/UObject/Property.cpp b/Engine/Source/Runtime/CoreUObject/Private/UObject/Property.cpp
--- a/Engine/Source/Runtime/CoreUObject/Private/UObject/Property.cpp
+++ b/Engine/Source/Runtime/CoreUObject/Private/UObject/Property.cpp
@@ -527,13 +527,14 @@ void UProperty::ExportCppDeclaration(FOutputDevice& Out, EExportedDeclaration::T
 		}
 	}
 
-	if( dynamic_cast<const UBoolProperty*>(this) )
+	if(auto BoolProperty = dynamic_cast<const UBoolProperty*>(this) )
 	{
 		// if this is a member variable, export it as a bitfield
 		if( ArrayDim==1 && DeclarationType == EExportedDeclaration::Member )
 		{
+			bool bCanUseBitfield = !BoolProperty->IsNativeBool();
 			// export as a uint32 member....bad to hardcode, but this is a special case that won't be used anywhere else
-			Out.Logf(TEXT("%s%s %s%s:1"), *TypeText, *ExtendedTypeText, *NameCpp, ArrayStr);
+			Out.Logf(TEXT("%s%s %s%s%s"), *TypeText, *ExtendedTypeText, *NameCpp, ArrayStr, bCanUseBitfield ? TEXT(":1") : TEXT(""));
 		}
 
 		//@todo we currently can't have out bools.. so this isn't really necessary, but eventually out bools may be supported, so leave here for now
