KismetCompilerCppBackend - fix for GetPrefixCPP

diff --git a/Engine/Source/Editor/KismetCompiler/Private/KismetCompilerCppBackend.cpp b/Engine/Source/Editor/KismetCompiler/Private/KismetCompilerCppBackend.cpp
--- a/Engine/Source/Editor/KismetCompiler/Private/KismetCompilerCppBackend.cpp
+++ b/Engine/Source/Editor/KismetCompiler/Private/KismetCompilerCppBackend.cpp
@@ -164,7 +164,7 @@ FString FKismetCppBackend::TermToText(const FBPTerminal* Term, const UProperty*
 		{
 			if (auto FoundClass = Cast<const UClass>(Term->ObjectLiteral))
 			{
-				return FString::Printf(TEXT("%s%s::StaticClass()"), *FoundClass->GetPrefixCPP(), *FoundClass->GetName());
+				return FString::Printf(TEXT("%s%s::StaticClass()"), FoundClass->GetPrefixCPP(), *FoundClass->GetName());
 			}
 			return FString(TEXT("NULL"));
 		}
@@ -190,7 +190,7 @@ FString FKismetCppBackend::TermToText(const FBPTerminal* Term, const UProperty*
 			{
 				if (auto LiteralClass = Cast<const UClass>(Term->ObjectLiteral))
 				{
-					return FString::Printf(TEXT("%s%s::StaticClass()"), *LiteralClass->GetPrefixCPP(), *LiteralClass->GetName());
+					return FString::Printf(TEXT("%s%s::StaticClass()"), LiteralClass->GetPrefixCPP(), *LiteralClass->GetName());
 				}
 				else
 				{
