UserDefinedStruct - properties won't be (unnecessarily) reloaded by linker

diff --git a/Engine/Source/Editor/KismetCompiler/Private/UserDefinedStructureCompilerUtils.cpp b/Engine/Source/Editor/KismetCompiler/Private/UserDefinedStructureCompilerUtils.cpp
--- a/Engine/Source/Editor/KismetCompiler/Private/UserDefinedStructureCompilerUtils.cpp
+++ b/Engine/Source/Editor/KismetCompiler/Private/UserDefinedStructureCompilerUtils.cpp
@@ -137,6 +137,7 @@ struct FUserDefinedStructureCompilerInner
 			UProperty* NewProperty = FKismetCompilerUtilities::CreatePropertyOnScope(Struct, VarDesc.VarName, VarType, NULL, 0, Schema, MessageLog);
 			if (NewProperty != NULL)
 			{
+				NewProperty->SetFlags(RF_LoadCompleted);
 				FKismetCompilerUtilities::LinkAddedProperty(Struct, NewProperty);
 			}
 			else
