Uber Graph Persistent Frame - All local variables, used by event graphs (see FKismetFunctionContext::EventGraphLocals), are no longer stored in class  scope. They are stored on a frame (as local variables in regular functions). For the ExecuteUbergraph_*() functions the FFrameÆs allocation is kept for the life of the object.

diff --git a/Engine/Config/BaseEngine.ini b/Engine/Config/BaseEngine.ini
--- a/Engine/Config/BaseEngine.ini
+++ b/Engine/Config/BaseEngine.ini
@@ -1637,6 +1637,8 @@ CompileDisplaysAnimBlueprintBackend=false
 bTurnOffEditorConstructionScript=false
 PrintStringDuration=2.0
 bUseLocalGraphVariables=false
+bPersistentUberGraphFrame=false
+bDisplaysLayout=false
 
 [/Script/Engine.Blueprint]
 bRecompileOnLoad=true
diff --git a/Engine/Source/Editor/BlueprintGraph/Classes/K2Node_SetVariableOnPersistentFrame.h b/Engine/Source/Editor/BlueprintGraph/Classes/K2Node_SetVariableOnPersistentFrame.h
new file mode 100644
--- /dev/null
+++ b/Engine/Source/Editor/BlueprintGraph/Classes/K2Node_SetVariableOnPersistentFrame.h
@@ -0,0 +1,26 @@
+// Copyright 1998-2014 Epic Games, Inc. All Rights Reserved.
+
+
+#pragma once
+#include "K2Node_SetVariableOnPersistentFrame.generated.h"
+
+/*
+ *	FOR INTERNAL USAGE ONLY!
+ */
+
+UCLASS(MinimalAPI)
+class UK2Node_SetVariableOnPersistentFrame : public UK2Node
+{
+	GENERATED_BODY()
+
+public:
+	// Begin UEdGraphNode interface
+	virtual void AllocateDefaultPins() override;
+	// End UEdGraphNode interface
+
+	// Begin K2Node interface
+	virtual bool IsNodePure() const override { return false; }
+	virtual class FNodeHandlingFunctor* CreateNodeHandler(class FKismetCompilerContext& CompilerContext) const override;
+	// End K2Node interface
+};
+
diff --git a/Engine/Source/Editor/BlueprintGraph/Private/K2Node_SetVariableOnPersistentFrame.cpp b/Engine/Source/Editor/BlueprintGraph/Private/K2Node_SetVariableOnPersistentFrame.cpp
new file mode 100644
--- /dev/null
+++ b/Engine/Source/Editor/BlueprintGraph/Private/K2Node_SetVariableOnPersistentFrame.cpp
@@ -0,0 +1,94 @@
+// Copyright 1998-2014 Epic Games, Inc. All Rights Reserved.
+
+
+#include "BlueprintGraphPrivatePCH.h"
+#include "K2Node_SetVariableOnPersistentFrame.h"
+#include "KismetCompiler.h"
+
+#define LOCTEXT_NAMESPACE "K2Node_SetVariableOnPersistentFrame"
+
+class FKCHandler_SetVariableOnPersistentFrame : public FNodeHandlingFunctor
+{
+public:
+	FKCHandler_SetVariableOnPersistentFrame(FKismetCompilerContext& InCompilerContext)
+		: FNodeHandlingFunctor(InCompilerContext)
+	{}
+
+	virtual void RegisterNets(FKismetFunctionContext& Context, UEdGraphNode* Node) override
+	{
+		check(Node && CompilerContext.GetSchema());
+		for (auto Pin : Node->Pins)
+		{
+			if (!Pin || CompilerContext.GetSchema()->IsMetaPin(*Pin))
+			{
+				continue;
+			}
+
+			if ((Pin->Direction != EGPD_Input) || (1 != Pin->LinkedTo.Num()) || Context.bIsUbergraph || !Context.NewClass || !Context.NewClass->UberGraphFunction)
+			{
+				CompilerContext.MessageLog.Error(*LOCTEXT("SetVariableOnPersistentFrame_IceError", "ICE SetVariableOnPersistentFrame @@").ToString(), Pin);
+				return;
+			}
+
+			UProperty* BoundProperty = FKismetCompilerUtilities::FindPropertyInScope(Context.NewClass->UberGraphFunction, Pin, CompilerContext.MessageLog, CompilerContext.GetSchema(), Context.NewClass);
+			if (!BoundProperty || (BoundProperty->GetOuter() != Context.NewClass->UberGraphFunction))
+			{
+				CompilerContext.MessageLog.Error(*LOCTEXT("SetVariableOnPersistentFrame_IceErrorNoProperty", "ICE SetVariableOnPersistentFrame - No property found. @@").ToString(), Pin);
+				return;
+			}
+
+			// Create the term in the list
+			FBPTerminal* Term = new (Context.PersistentFrameVariableReferences)FBPTerminal();
+			Term->CopyFromPin(Pin, Pin->PinName);
+			Term->AssociatedVarProperty = BoundProperty;
+			Context.NetMap.Add(Pin, Term);
+		}
+	}
+
+	virtual void Compile(FKismetFunctionContext& Context, UEdGraphNode* Node) override
+	{
+		check(Node && CompilerContext.GetSchema());
+		for (auto Pin : Node->Pins)
+		{
+			if (!Pin || CompilerContext.GetSchema()->IsMetaPin(*Pin))
+			{
+				continue;
+			}
+
+			FBPTerminal** DestTerm = Context.NetMap.Find(Pin);
+			FBPTerminal** SourceTerm = Context.NetMap.Find(FEdGraphUtilities::GetNetFromPin(Pin));
+			if (DestTerm && *DestTerm && SourceTerm && *SourceTerm)
+			{
+				FBlueprintCompiledStatement& Statement = Context.AppendStatementForNode(Node);
+				Statement.Type = KCST_AssignmentOnPersistentFrame;
+				Statement.LHS = *DestTerm;
+				Statement.RHS.Add(*SourceTerm);
+			}
+			else
+			{
+				CompilerContext.MessageLog.Error(*LOCTEXT("SetVariableOnPersistentFrame_NoTerm", "ICE SetVariableOnPersistentFrame - No terminal found. @@").ToString(), Pin);
+				return;
+			}
+		}
+
+		// Generate the output impulse from this node
+		GenerateSimpleThenGoto(Context, *Node);
+	}
+};
+
+void UK2Node_SetVariableOnPersistentFrame::AllocateDefaultPins()
+{
+	const UEdGraphSchema_K2* K2Schema = GetDefault<UEdGraphSchema_K2>();
+
+	CreatePin(EGPD_Input, K2Schema->PC_Exec, TEXT(""), NULL, false, false, K2Schema->PN_Execute);
+	CreatePin(EGPD_Output, K2Schema->PC_Exec, TEXT(""), NULL, false, false, K2Schema->PN_Then);
+
+	Super::AllocateDefaultPins();
+}
+
+FNodeHandlingFunctor* UK2Node_SetVariableOnPersistentFrame::CreateNodeHandler(FKismetCompilerContext& CompilerContext) const
+{
+	return new FKCHandler_SetVariableOnPersistentFrame(CompilerContext);
+}
+
+#undef LOCTEXT_NAMESPACE
diff --git a/Engine/Source/Editor/KismetCompiler/Private/KismetCompiler.cpp b/Engine/Source/Editor/KismetCompiler/Private/KismetCompiler.cpp
--- a/Engine/Source/Editor/KismetCompiler/Private/KismetCompiler.cpp
+++ b/Engine/Source/Editor/KismetCompiler/Private/KismetCompiler.cpp
@@ -18,6 +18,7 @@
 #include "UserDefinedStructureCompilerUtils.h"
 #include "EditorCategoryUtils.h"
 #include "K2Node_EnumLiteral.h"
+#include "K2Node_SetVariableOnPersistentFrame.h"
 #include "EdGraph/EdGraphNode_Documentation.h"
 #include "Engine/DynamicBlueprintBinding.h"
 #include "Kismet/KismetMathLibrary.h"
@@ -605,7 +606,8 @@ void FKismetCompilerContext::CreatePropertiesFromList(UStruct* Scope, UField**&
 
 			if (Term.bPassedByReference)
 			{
-				if (!NewProperty->HasAnyPropertyFlags(CPF_OutParm))
+				// special case for BlueprintImplementableEvent
+				if (NewProperty->HasAnyPropertyFlags(CPF_Parm) && !NewProperty->HasAnyPropertyFlags(CPF_OutParm))
 				{
 					NewProperty->SetPropertyFlags(CPF_OutParm | CPF_ReferenceParm);
 				}
@@ -698,6 +700,10 @@ void FKismetCompilerContext::CreatePropertiesFromList(UStruct* Scope, UField**&
 
 void FKismetCompilerContext::CreateLocalVariablesForFunction(FKismetFunctionContext& Context)
 {
+	ensure(Context.IsEventGraph() || !Context.EventGraphLocals.Num());
+	ensure(!Context.IsEventGraph() || !Context.Locals.Num());
+
+	const bool bPersistentUberGraphFrame = UBlueprintGeneratedClass::UsePersistentUberGraphFrame() && Context.bIsUbergraph;
 	// Local stack frame (or maybe class for the ubergraph)
 	{
 		const bool bArePropertiesLocal = true;
@@ -710,6 +716,11 @@ void FKismetCompilerContext::CreateLocalVariablesForFunction(FKismetFunctionCont
 		CreatePropertiesFromList(Context.Function, PropertyStorageLocation, Context.Results, CPF_Parm | CPF_OutParm, bArePropertiesLocal, /*bPropertiesAreParameters=*/ true);
 		CreatePropertiesFromList(Context.Function, PropertyStorageLocation, Context.Locals, 0, bArePropertiesLocal, /*bPropertiesAreParameters=*/ true);
 
+		if (bPersistentUberGraphFrame)
+		{
+			CreatePropertiesFromList(Context.Function, PropertyStorageLocation, Context.EventGraphLocals, 0, bArePropertiesLocal, true);
+		}
+
 		// If there were local properties, place them at the end of the property storage location
 		if(LocalProperties)
 		{
@@ -746,8 +757,6 @@ void FKismetCompilerContext::CreateLocalVariablesForFunction(FKismetFunctionCont
 
 	// Class
 	{
-		const bool bArePropertiesLocal = false;
-		
 		int32 PropertySafetyCounter = 100000;
 		UField** PropertyStorageLocation = &(NewClass->Children);
 		while (*PropertyStorageLocation != NULL)
@@ -760,8 +769,12 @@ void FKismetCompilerContext::CreateLocalVariablesForFunction(FKismetFunctionCont
 			PropertyStorageLocation = &((*PropertyStorageLocation)->Next);
 		}
 
+		const bool bArePropertiesLocal = false;
 		const uint64 UbergraphHiddenVarFlags = CPF_Transient | CPF_DuplicateTransient;
-		CreatePropertiesFromList(NewClass, PropertyStorageLocation, Context.EventGraphLocals, UbergraphHiddenVarFlags, bArePropertiesLocal);
+		if (!bPersistentUberGraphFrame)
+		{
+			CreatePropertiesFromList(NewClass, PropertyStorageLocation, Context.EventGraphLocals, UbergraphHiddenVarFlags, bArePropertiesLocal);
+		}
 
 		// Handle level actor references
 		const uint64 LevelActorReferenceVarFlags = 0/*CPF_Edit*/;
@@ -1165,6 +1178,11 @@ void FKismetCompilerContext::PrecompileFunction(FKismetFunctionContext& Context)
 
 		// Add the function to it's owner class function name -> function map
 		Context.NewClass->AddFunctionToFunctionMap(Context.Function);
+		if (UBlueprintGeneratedClass::UsePersistentUberGraphFrame() && Context.bIsUbergraph)
+		{
+			ensure(!NewClass->UberGraphFunction);
+			NewClass->UberGraphFunction = Context.Function;
+		}
 
 		// Create any user defined variables, this must occur before registering nets so that the properties are in place
 		UField** PropertyStorageLocation = &(Context.Function->Children);
@@ -2328,7 +2346,7 @@ void FKismetCompilerContext::CreateFunctionStubForEvent(UK2Node_Event* SrcEventN
 	}
 
 	// Copy each event parameter to the assignment node, if there are any inputs
-	UK2Node_VariableSet* AssignmentNode = NULL;
+	UK2Node* AssignmentNode = NULL;
 	for (int32 PinIndex = 0; PinIndex < EntryNode->Pins.Num(); ++PinIndex)
 	{
 		UEdGraphPin* SourcePin = EntryNode->Pins[PinIndex];
@@ -2337,8 +2355,16 @@ void FKismetCompilerContext::CreateFunctionStubForEvent(UK2Node_Event* SrcEventN
 			if (AssignmentNode == NULL)
 			{
 				// Create a variable write node to store the parameters into the ubergraph frame storage
-				AssignmentNode = SpawnIntermediateNode<UK2Node_VariableSet>(SrcEventNode, ChildStubGraph);
-				AssignmentNode->VariableReference.SetSelfMember(NAME_None);
+				if (UBlueprintGeneratedClass::UsePersistentUberGraphFrame())
+				{
+					AssignmentNode = SpawnIntermediateNode<UK2Node_SetVariableOnPersistentFrame>(SrcEventNode, ChildStubGraph);
+				}
+				else
+				{
+					auto VariableSetNode = SpawnIntermediateNode<UK2Node_VariableSet>(SrcEventNode, ChildStubGraph);
+					VariableSetNode->VariableReference.SetSelfMember(NAME_None);
+					AssignmentNode = VariableSetNode;
+				}
 				AssignmentNode->AllocateDefaultPins();
 			}
 
@@ -2558,15 +2584,6 @@ void FKismetCompilerContext::CreateAndProcessUbergraph()
 	// Merge all of the top-level pages
 	MergeUbergraphPagesIn(ConsolidatedEventGraph);
 
-	// Add a dummy entry point to the uber graph, to get the function signature correct
-	{
-		UK2Node_FunctionEntry* EntryNode = SpawnIntermediateNode<UK2Node_FunctionEntry>(NULL, ConsolidatedEventGraph);
-		EntryNode->SignatureClass = UObject::StaticClass();
-		EntryNode->SignatureName = Schema->FN_ExecuteUbergraphBase;
-		EntryNode->CustomGeneratedFunctionName = ConsolidatedEventGraph->GetFName();
-		EntryNode->AllocateDefaultPins();
-	}
-
 	// Loop over implemented interfaces, and add dummy event entry points for events that aren't explicitly handled by the user
 	TArray<UK2Node_Event*> EntryPoints;
 	ConsolidatedEventGraph->GetNodesOfClass(EntryPoints);
@@ -2607,40 +2624,51 @@ void FKismetCompilerContext::CreateAndProcessUbergraph()
 		}
 	}
 
-	// Expand out nodes that need it
-	ExpansionStep(ConsolidatedEventGraph, true);
-
-	// If a function in the graph cannot be overridden/placed as event make sure that it is not.
-	VerifyValidOverrideEvent(ConsolidatedEventGraph);
-
-	// Do some cursory validation (pin types match, inputs to outputs, pins never point to their parent node, etc...)
+	if (ConsolidatedEventGraph->Nodes.Num())
 	{
-		UbergraphContext = new (FunctionList) FKismetFunctionContext(MessageLog, Schema, NewClass, Blueprint);
-		UbergraphContext->SourceGraph = ConsolidatedEventGraph;
-		UbergraphContext->MarkAsEventGraph();
-		UbergraphContext->MarkAsInternalOrCppUseOnly();
-		UbergraphContext->SetExternalNetNameMap(&ClassScopeNetNameMap);
+		// Add a dummy entry point to the uber graph, to get the function signature correct
+		{
+			UK2Node_FunctionEntry* EntryNode = SpawnIntermediateNode<UK2Node_FunctionEntry>(NULL, ConsolidatedEventGraph);
+			EntryNode->SignatureClass = UObject::StaticClass();
+			EntryNode->SignatureName = Schema->FN_ExecuteUbergraphBase;
+			EntryNode->CustomGeneratedFunctionName = ConsolidatedEventGraph->GetFName();
+			EntryNode->AllocateDefaultPins();
+		}
 
-		Blueprint->EventGraphs.Empty();
+		// Expand out nodes that need it
+		ExpansionStep(ConsolidatedEventGraph, true);
 
-		// Validate all the nodes in the graph
-		for (int32 ChildIndex = 0; ChildIndex < ConsolidatedEventGraph->Nodes.Num(); ++ChildIndex)
+		// If a function in the graph cannot be overridden/placed as event make sure that it is not.
+		VerifyValidOverrideEvent(ConsolidatedEventGraph);
+
+		// Do some cursory validation (pin types match, inputs to outputs, pins never point to their parent node, etc...)
 		{
-			const UEdGraphNode* Node = ConsolidatedEventGraph->Nodes[ChildIndex];	
-			const int32 SavedErrorCount = MessageLog.NumErrors;
-			ValidateNode(Node);
+			UbergraphContext = new (FunctionList)FKismetFunctionContext(MessageLog, Schema, NewClass, Blueprint);
+			UbergraphContext->SourceGraph = ConsolidatedEventGraph;
+			UbergraphContext->MarkAsEventGraph();
+			UbergraphContext->MarkAsInternalOrCppUseOnly();
+			UbergraphContext->SetExternalNetNameMap(&ClassScopeNetNameMap);
+
+			Blueprint->EventGraphs.Empty();
 
-			// If the node didn't generate any errors then generate function stubs for event entry nodes etc.
-			if (SavedErrorCount == MessageLog.NumErrors)
+			// Validate all the nodes in the graph
+			for (int32 ChildIndex = 0; ChildIndex < ConsolidatedEventGraph->Nodes.Num(); ++ChildIndex)
 			{
-				if (UK2Node_Event* SrcEventNode = Cast<UK2Node_Event>(ConsolidatedEventGraph->Nodes[ChildIndex]))
+				const UEdGraphNode* Node = ConsolidatedEventGraph->Nodes[ChildIndex];
+				const int32 SavedErrorCount = MessageLog.NumErrors;
+				ValidateNode(Node);
+
+				// If the node didn't generate any errors then generate function stubs for event entry nodes etc.
+				if (SavedErrorCount == MessageLog.NumErrors)
 				{
-					CreateFunctionStubForEvent(SrcEventNode, Blueprint);
+					if (UK2Node_Event* SrcEventNode = Cast<UK2Node_Event>(ConsolidatedEventGraph->Nodes[ChildIndex]))
+					{
+						CreateFunctionStubForEvent(SrcEventNode, Blueprint);
+					}
 				}
 			}
 		}
 	}
-
 }
 
 void FKismetCompilerContext::AutoAssignNodePosition(UEdGraphNode* Node)
@@ -2680,6 +2708,11 @@ void FKismetCompilerContext::AdvanceMacroPlacement(int32 Width, int32 Height)
 
 void FKismetCompilerContext::CreateCommentBlockAroundNodes(const TArray<UEdGraphNode*>& Nodes, UObject* SourceObject, UEdGraph* TargetGraph, FString CommentText, FLinearColor CommentColor, int32& Out_OffsetX, int32& Out_OffsetY)
 {
+	if (!Nodes.Num())
+	{
+		return;
+	}
+
 	FIntRect Bounds = FEdGraphUtilities::CalculateApproximateNodeBoundaries(Nodes);
 
 	// Figure out how to offset the expanded nodes to fit into our tile
@@ -3259,6 +3292,14 @@ void FKismetCompilerContext::Compile()
 		}
 	}
 
+	if (UBlueprintGeneratedClass::UsePersistentUberGraphFrame() && UbergraphContext)
+	{
+		//UBER GRAPH PERSISTENT FRAME
+		FEdGraphPinType Type(TEXT("struct"), TEXT(""), FPointerToUberGraphFrame::StaticStruct(), false, false);
+		auto Property = CreateVariable(UBlueprintGeneratedClass::GetUberGraphFrameName(), Type);
+		Property->SetPropertyFlags(CPF_DuplicateTransient | CPF_Transient);
+	}
+
 	{ BP_SCOPED_COMPILER_EVENT_STAT(EKismetCompilerStats_BindAndLinkClass);
 
 		// Relink the class
@@ -3470,6 +3511,26 @@ void FKismetCompilerContext::Compile()
 				}
 			}
 		}
+
+		static const FBoolConfigValueHelper DisplayLayout(TEXT("Kismet"), TEXT("bDisplaysLayout"), GEngineIni);
+		if (!Blueprint->bIsRegeneratingOnLoad && bIsFullCompile && DisplayLayout && NewClass)
+		{
+			UE_LOG(LogK2Compiler, Log, TEXT("\n\nLAYOUT CLASS %s:"), *GetNameSafe(NewClass));
+
+			for (auto Prop : TFieldRange<UProperty>(NewClass, EFieldIteratorFlags::ExcludeSuper))
+			{
+				UE_LOG(LogK2Compiler, Log, TEXT("%5d:\t%-64s\t%s"), Prop->GetOffset_ForGC(), *GetNameSafe(Prop), *Prop->GetCPPType());
+			}
+
+			if (NewClass->UberGraphFunction)
+			{
+				UE_LOG(LogK2Compiler, Log, TEXT("\n\nLAYOUT FRAME %s:"), *GetNameSafe(NewClass->UberGraphFunction));
+				for (auto Prop : TFieldRange<UProperty>(NewClass->UberGraphFunction))
+				{
+					UE_LOG(LogK2Compiler, Log, TEXT("%5d:\t%-64s\t%s"), Prop->GetOffset_ForGC(), *GetNameSafe(Prop), *Prop->GetCPPType());
+				}
+			}
+		}
 	}
 
 	// If this was a skeleton compile, make sure everything is RF_Transient
diff --git a/Engine/Source/Editor/KismetCompiler/Private/KismetCompilerVMBackend.cpp b/Engine/Source/Editor/KismetCompiler/Private/KismetCompilerVMBackend.cpp
--- a/Engine/Source/Editor/KismetCompiler/Private/KismetCompilerVMBackend.cpp
+++ b/Engine/Source/Editor/KismetCompiler/Private/KismetCompilerVMBackend.cpp
@@ -849,6 +849,18 @@ public:
 		EmitTerm(SourceExpression, DestinationExpression->AssociatedVarProperty);
 	}
 
+	void EmitAssignmentOnPersistentFrameStatment(FBlueprintCompiledStatement& Statement)
+	{
+		FBPTerminal* DestinationExpression = Statement.LHS;
+		FBPTerminal* SourceExpression = Statement.RHS[0];
+
+		Writer << Ex_LetValueOnPersistentFrame;
+		check(ClassBeingBuilt && ClassBeingBuilt->UberGraphFunction);
+		Writer << DestinationExpression->AssociatedVarProperty;
+
+		EmitTerm(SourceExpression, DestinationExpression->AssociatedVarProperty);
+	}
+
 	void EmitCastObjToInterfaceStatement(FBlueprintCompiledStatement& Statement)
 	{
 		FBPTerminal* DestinationExpression = Statement.LHS;
@@ -1161,6 +1173,9 @@ public:
 		case KCST_Assignment:
 			EmitAssignmentStatment(Statement);
 			break;
+		case KCST_AssignmentOnPersistentFrame:
+			EmitAssignmentOnPersistentFrameStatment(Statement);
+			break;
 		case KCST_CastObjToInterface:
 			EmitCastObjToInterfaceStatement(Statement);
 			break;
diff --git a/Engine/Source/Editor/KismetCompiler/Public/BlueprintCompiledStatement.h b/Engine/Source/Editor/KismetCompiler/Public/BlueprintCompiledStatement.h
--- a/Engine/Source/Editor/KismetCompiler/Public/BlueprintCompiledStatement.h
+++ b/Engine/Source/Editor/KismetCompiler/Public/BlueprintCompiledStatement.h
@@ -32,6 +32,7 @@ enum EKismetCompiledStatementType
 	KCST_CreateArray = 22,		// Creates and sets an array literal term
 	KCST_CrossInterfaceCast = 23,	// TargetInterface(Interface)
 	KCST_MetaCast = 24,	    // Cast<TargetClass>(TargetObject)
+	KCST_AssignmentOnPersistentFrame = 25, //
 };
 
 //@TODO: Too rigid / icky design
diff --git a/Engine/Source/Editor/KismetCompiler/Public/KismetCompiledFunctionContext.h b/Engine/Source/Editor/KismetCompiler/Public/KismetCompiledFunctionContext.h
--- a/Engine/Source/Editor/KismetCompiler/Public/KismetCompiledFunctionContext.h
+++ b/Engine/Source/Editor/KismetCompiler/Public/KismetCompiledFunctionContext.h
@@ -65,6 +65,7 @@ public:
 	TIndirectArray<FBPTerminal> Parameters;
 	TIndirectArray<FBPTerminal> Results;
 	TIndirectArray<FBPTerminal> VariableReferences;
+	TIndirectArray<FBPTerminal> PersistentFrameVariableReferences;
 	TIndirectArray<FBPTerminal> Literals;
 	TIndirectArray<FBPTerminal> Locals;
 	TIndirectArray<FBPTerminal> EventGraphLocals;
diff --git a/Engine/Source/Editor/UnrealEd/Private/ScriptDisassembler.cpp b/Engine/Source/Editor/UnrealEd/Private/ScriptDisassembler.cpp
--- a/Engine/Source/Editor/UnrealEd/Private/ScriptDisassembler.cpp
+++ b/Engine/Source/Editor/UnrealEd/Private/ScriptDisassembler.cpp
@@ -275,6 +275,22 @@ void FKismetBytecodeDisassembler::ProcessCommon(int32& ScriptIndex, EExprToken O
 			SerializeExpr( ScriptIndex );
 
 			DropIndent();
+			break;
+		}
+	case Ex_LetValueOnPersistentFrame:
+		{
+			Ar.Logf(TEXT("%s $%X: LetValueOnPersistentFrame"), *Indents, (int32)Opcode);
+			AddIndent();
+
+			auto Prop = ReadPointer<UProperty>(ScriptIndex);
+			Ar.Logf(TEXT("%s Destination variable: %s, offset: %d"), *Indents, *GetNameSafe(Prop), 
+				Prop ? Prop->GetOffset_ForDebug() : 0);
+			
+			Ar.Logf(TEXT("%s Expression:"), *Indents);
+			SerializeExpr(ScriptIndex);
+
+			DropIndent();
+
 			break;
 		}
 	case EX_StructMemberContext:
diff --git a/Engine/Source/Runtime/CoreUObject/Private/Blueprint/BlueprintSupport.cpp b/Engine/Source/Runtime/CoreUObject/Private/Blueprint/BlueprintSupport.cpp
--- a/Engine/Source/Runtime/CoreUObject/Private/Blueprint/BlueprintSupport.cpp
+++ b/Engine/Source/Runtime/CoreUObject/Private/Blueprint/BlueprintSupport.cpp
@@ -177,6 +177,7 @@ bool UObject::IsInBlueprint() const
 void UObject::DestroyNonNativeProperties()
 {
 	// Destroy properties that won't be destroyed by the native destructor
+	GetClass()->DestroyPersistentUberGraphFrame(this);
 	{
 		for (UProperty* P = GetClass()->DestructorLink; P; P = P->DestructorLinkNext)
 		{
diff --git a/Engine/Source/Runtime/CoreUObject/Private/UObject/ScriptCore.cpp b/Engine/Source/Runtime/CoreUObject/Private/UObject/ScriptCore.cpp
--- a/Engine/Source/Runtime/CoreUObject/Private/UObject/ScriptCore.cpp
+++ b/Engine/Source/Runtime/CoreUObject/Private/UObject/ScriptCore.cpp
@@ -402,9 +402,13 @@ void UObject::CallFunction( FFrame& Stack, RESULT_DECL, UFunction* Function )
 	}
 	else
 	{
-		// Make new stack frame in the current context.
-		uint8* Frame = (uint8*)FMemory_Alloca(Function->PropertiesSize);
-		FMemory::Memzero( Frame, Function->PropertiesSize );
+		uint8* Frame = GetClass()->GetPersistentUberGraphFrame(this, Function);
+		const bool bUsePersistentFrame = (NULL != Frame);
+		if (!bUsePersistentFrame)
+		{
+			Frame = (uint8*)FMemory_Alloca(Function->PropertiesSize);
+			FMemory::Memzero(Frame, Function->PropertiesSize);
+		}
 		FFrame NewStack( this, Function, Frame, &Stack, Function->Children );
 		FOutParmRec** LastOut = &NewStack.OutParms;
 		UProperty* Property;
@@ -457,7 +461,7 @@ void UObject::CallFunction( FFrame& Stack, RESULT_DECL, UFunction* Function )
 				// warning: Stack.MostRecentPropertyAddress could be NULL for optional out parameters
 				// if that's the case, we use the extra memory allocated for the out param in the function's locals
 				// so there's always a valid address
-				Out->PropAddr = (Stack.MostRecentPropertyAddress != NULL) ? Stack.MostRecentPropertyAddress : Property->ContainerPtrToValuePtr<uint8>(Frame);
+				Out->PropAddr = (Stack.MostRecentPropertyAddress != NULL) ? Stack.MostRecentPropertyAddress : Property->ContainerPtrToValuePtr<uint8>(NewStack.Locals);
 				Out->Property = Property;
 
 				// add the new out param info to the stack frame's linked list
@@ -489,10 +493,13 @@ void UObject::CallFunction( FFrame& Stack, RESULT_DECL, UFunction* Function )
 		}
 #endif
 
-		// Initialize any local struct properties with defaults
-		for ( UProperty* LocalProp = Function->FirstPropertyToInit; LocalProp != NULL; LocalProp = (UProperty*)LocalProp->Next )
+		if (!bUsePersistentFrame)
 		{
-			LocalProp->InitializeValue_InContainer(NewStack.Locals);
+			// Initialize any local struct properties with defaults
+			for (UProperty* LocalProp = Function->FirstPropertyToInit; LocalProp != NULL; LocalProp = (UProperty*)LocalProp->Next)
+			{
+				LocalProp->InitializeValue_InContainer(NewStack.Locals);
+			}
 		}
 
 		const bool bIsValidFunction = (Function->FunctionFlags & FUNC_Native) || (Function->Script.Num() > 0);
@@ -503,12 +510,15 @@ void UObject::CallFunction( FFrame& Stack, RESULT_DECL, UFunction* Function )
 			ProcessInternal( NewStack, Result );
 		}
 
-		// destruct properties on the stack, except for out params since we know we didn't use that memory
-		for (UProperty* Destruct = Function->DestructorLink; Destruct; Destruct = Destruct->DestructorLinkNext)
+		if (!bUsePersistentFrame)
 		{
-			if (!Destruct->HasAnyPropertyFlags(CPF_OutParm))
+			// destruct properties on the stack, except for out params since we know we didn't use that memory
+			for (UProperty* Destruct = Function->DestructorLink; Destruct; Destruct = Destruct->DestructorLinkNext)
 			{
-				Destruct->DestroyValue_InContainer(NewStack.Locals);
+				if (!Destruct->HasAnyPropertyFlags(CPF_OutParm))
+				{
+					Destruct->DestroyValue_InContainer(NewStack.Locals);
+				}
 			}
 		}
 	}
@@ -829,15 +839,24 @@ void UObject::ProcessEvent( UFunction* Function, void* Parms )
 
 	// Scope required for scoped script stats.
 	{
+		uint8* Frame = GetClass()->GetPersistentUberGraphFrame(this, Function);
+		const bool bUsePersistentFrame = (NULL != Frame);
+		if (!bUsePersistentFrame)
+		{
+			Frame = (uint8*)FMemory_Alloca(Function->PropertiesSize);
+			// zero the local property memory
+			FMemory::Memzero(Frame + Function->ParmsSize, Function->PropertiesSize - Function->ParmsSize);
+		}
+
+		// initialize the parameter properties
+		FMemory::Memcpy(Frame, Parms, Function->ParmsSize);
+
 		// Create a new local execution stack.
-		FFrame NewStack( this, Function, FMemory_Alloca(Function->PropertiesSize), NULL, Function->Children );
+		FFrame NewStack(this, Function, Frame, NULL, Function->Children);
+
 		checkSlow(NewStack.Locals || Function->ParmsSize == 0);
 
-		// initialize the parameter properties
-		FMemory::Memcpy( NewStack.Locals, Parms, Function->ParmsSize );
 
-		// zero the local property memory
-		FMemory::Memzero( NewStack.Locals+Function->ParmsSize, Function->PropertiesSize-Function->ParmsSize );
 
 		// if the function has out parameters, fill the stack frame's out parameter info with the info for those params 
 		if ( Function->HasAnyFunctionFlags(FUNC_HasOutParms) )
@@ -879,10 +898,14 @@ void UObject::ProcessEvent( UFunction* Function, void* Parms )
 #endif
 		}
 
-		for ( UProperty* LocalProp = Function->FirstPropertyToInit; LocalProp != NULL; LocalProp = (UProperty*)LocalProp->Next )
+		if (!bUsePersistentFrame)
 		{
-			LocalProp->InitializeValue_InContainer(NewStack.Locals);
+			for (UProperty* LocalProp = Function->FirstPropertyToInit; LocalProp != NULL; LocalProp = (UProperty*)LocalProp->Next)
+			{
+				LocalProp->InitializeValue_InContainer(NewStack.Locals);
+			}
 		}
+
 		// Call native function or UObject::ProcessInternal.
 		if (Function->FunctionFlags & FUNC_Native)
 		{
@@ -895,17 +918,20 @@ void UObject::ProcessEvent( UFunction* Function, void* Parms )
 			Function->Invoke(this, NewStack, (uint8*)Parms + Function->ReturnValueOffset);
 		}
 
-		// Destroy local variables except function parameters.!! see also UObject::CallFunctionByNameWithArguments
-		// also copy back constructed value parms here so the correct copy is destroyed when the event function returns
-		for (UProperty* P = Function->DestructorLink; P; P = P->DestructorLinkNext)
+		if (!bUsePersistentFrame)
 		{
-			if (!P->IsInContainer(Function->ParmsSize))
-			{
-				P->DestroyValue_InContainer(NewStack.Locals);
-			}
-			else if (!(P->PropertyFlags & CPF_OutParm))
+			// Destroy local variables except function parameters.!! see also UObject::CallFunctionByNameWithArguments
+			// also copy back constructed value parms here so the correct copy is destroyed when the event function returns
+			for (UProperty* P = Function->DestructorLink; P; P = P->DestructorLinkNext)
 			{
-				FMemory::Memcpy(P->ContainerPtrToValuePtr<uint8>(Parms), P->ContainerPtrToValuePtr<uint8>(NewStack.Locals), P->ArrayDim * P->ElementSize);
+				if (!P->IsInContainer(Function->ParmsSize))
+				{
+					P->DestroyValue_InContainer(NewStack.Locals);
+				}
+				else if (!(P->PropertyFlags & CPF_OutParm))
+				{
+					FMemory::Memcpy(P->ContainerPtrToValuePtr<uint8>(Parms), P->ContainerPtrToValuePtr<uint8>(NewStack.Locals), P->ArrayDim * P->ElementSize);
+				}
 			}
 		}
 	}
@@ -1186,6 +1212,21 @@ void UObject::execPopExecutionFlowIfNot( FFrame& Stack, RESULT_DECL )
 }
 IMPLEMENT_VM_FUNCTION( EX_PopExecutionFlowIfNot, execPopExecutionFlowIfNot );
 
+void UObject::execLetValueOnPersistentFrame(FFrame& Stack, RESULT_DECL)
+{
+	Stack.MostRecentProperty = NULL;
+	Stack.MostRecentPropertyAddress = NULL;
+
+	auto DestProperty = Stack.ReadProperty();
+	checkSlow(DestProperty);
+	auto UberGraphFunction = CastChecked<UFunction>(DestProperty->GetOwnerStruct());
+	auto FrameBase = Stack.Object->GetClass()->GetPersistentUberGraphFrame(Stack.Object, UberGraphFunction);
+	checkSlow(FrameBase);
+	auto DestAddress = DestProperty->ContainerPtrToValuePtr<uint8>(FrameBase);
+
+	Stack.Step(Stack.Object, DestAddress);
+}
+IMPLEMENT_VM_FUNCTION(Ex_LetValueOnPersistentFrame, execLetValueOnPersistentFrame);
 
 void UObject::execLet( FFrame& Stack, RESULT_DECL )
 {
diff --git a/Engine/Source/Runtime/CoreUObject/Private/UObject/UObjectBase.cpp b/Engine/Source/Runtime/CoreUObject/Private/UObject/UObjectBase.cpp
--- a/Engine/Source/Runtime/CoreUObject/Private/UObject/UObjectBase.cpp
+++ b/Engine/Source/Runtime/CoreUObject/Private/UObject/UObjectBase.cpp
@@ -257,7 +257,9 @@ void UObjectBase::SetClass(UClass* NewClass)
 	STAT(StatID = TStatId();) // reset the stat id since this thing now has a different name
 
 	UnhashObject(this);
+	Class->DestroyPersistentUberGraphFrame((UObject*)this);
 	Class = NewClass;
+	Class->CreatePersistentUberGraphFrame((UObject*)this);
 	HashObject(this);
 }
 
diff --git a/Engine/Source/Runtime/CoreUObject/Private/UObject/UObjectGlobals.cpp b/Engine/Source/Runtime/CoreUObject/Private/UObject/UObjectGlobals.cpp
--- a/Engine/Source/Runtime/CoreUObject/Private/UObject/UObjectGlobals.cpp
+++ b/Engine/Source/Runtime/CoreUObject/Private/UObject/UObjectGlobals.cpp
@@ -1808,6 +1808,7 @@ void UObject::PostInitProperties()
 #if !(UE_BUILD_SHIPPING || UE_BUILD_TEST)
 	PostInitPropertiesCheck.Push(this);
 #endif
+	GetClass()->CreatePersistentUberGraphFrame(this);
 }
 
 UObject::UObject(const FObjectInitializer& ObjectInitializer)
diff --git a/Engine/Source/Runtime/CoreUObject/Public/UObject/Class.h b/Engine/Source/Runtime/CoreUObject/Public/UObject/Class.h
--- a/Engine/Source/Runtime/CoreUObject/Public/UObject/Class.h
+++ b/Engine/Source/Runtime/CoreUObject/Public/UObject/Class.h
@@ -1953,6 +1953,19 @@ public:
 	*/
 	FName GetDefaultObjectName();
 
+	virtual uint8* GetPersistentUberGraphFrame(UObject* Obj, UFunction* FuncToCheck) const
+	{
+		return NULL;
+	}
+
+	virtual void CreatePersistentUberGraphFrame(UObject* Obj) const
+	{
+	}
+
+	virtual void DestroyPersistentUberGraphFrame(UObject* Obj) const
+	{
+	}
+
 	/**
 	  * Get the default object from the class and cast to a particular type
 	  * @return		the CDO for this class
diff --git a/Engine/Source/Runtime/CoreUObject/Public/UObject/Script.h b/Engine/Source/Runtime/CoreUObject/Public/UObject/Script.h
--- a/Engine/Source/Runtime/CoreUObject/Public/UObject/Script.h
+++ b/Engine/Source/Runtime/CoreUObject/Public/UObject/Script.h
@@ -234,6 +234,7 @@ enum EExprToken
 	EX_BindDelegate			= 0x61, // bind object and name to delegate
 	EX_RemoveMulticastDelegate = 0x62, // Remove a delegate from a multicast delegate's targets
 	EX_CallMulticastDelegate = 0x63, // Call multicast delegate
+	Ex_LetValueOnPersistentFrame = 0x64,
 	EX_Max					= 0x100,
 };
 
diff --git a/Engine/Source/Runtime/CoreUObject/Public/UObject/ScriptSerialization.h b/Engine/Source/Runtime/CoreUObject/Public/UObject/ScriptSerialization.h
--- a/Engine/Source/Runtime/CoreUObject/Public/UObject/ScriptSerialization.h
+++ b/Engine/Source/Runtime/CoreUObject/Public/UObject/ScriptSerialization.h
@@ -129,6 +129,12 @@
 			SerializeExpr( iCode, Ar ); // Assignment expr.
 			break;
 		}
+		case Ex_LetValueOnPersistentFrame:
+		{
+			XFER_PROP_POINTER;			// Destination property.
+			SerializeExpr(iCode, Ar);	// Assignment expr.
+			break;
+		}
 		case EX_StructMemberContext:
 		{
 			XFERPTR(UProperty*);        // struct member expr.
diff --git a/Engine/Source/Runtime/CoreUObject/Public/UObject/UObject.h b/Engine/Source/Runtime/CoreUObject/Public/UObject/UObject.h
--- a/Engine/Source/Runtime/CoreUObject/Public/UObject/UObject.h
+++ b/Engine/Source/Runtime/CoreUObject/Public/UObject/UObject.h
@@ -974,6 +974,8 @@ public:
 	DECLARE_FUNCTION(execBindDelegate);
 	DECLARE_FUNCTION(execCallMulticastDelegate);
 
+	DECLARE_FUNCTION(execLetValueOnPersistentFrame);
+
 	// -- K2 support functions
 	struct Object_eventExecuteUbergraph_Parms
 	{
diff --git a/Engine/Source/Runtime/Engine/Classes/Engine/BlueprintGeneratedClass.h b/Engine/Source/Runtime/Engine/Classes/Engine/BlueprintGeneratedClass.h
--- a/Engine/Source/Runtime/Engine/Classes/Engine/BlueprintGeneratedClass.h
+++ b/Engine/Source/Runtime/Engine/Classes/Engine/BlueprintGeneratedClass.h
@@ -55,6 +55,34 @@ public:
 	}
 };
 
+USTRUCT()
+struct FPointerToUberGraphFrame
+{
+	GENERATED_USTRUCT_BODY()
+
+public:
+	uint8* RawPointer;
+
+	FPointerToUberGraphFrame() 
+		: RawPointer(NULL)
+	{}
+
+	~FPointerToUberGraphFrame()
+	{
+		check(!RawPointer);
+	}
+};
+
+template<>
+struct TStructOpsTypeTraits<FPointerToUberGraphFrame> : public TStructOpsTypeTraitsBase
+{
+	enum
+	{
+		WithZeroConstructor = true,
+		WithCopy = false,
+	};
+};
+
 //////////////////////////////////////////////////////////////////////////
 // TSimpleRingBuffer
 
@@ -406,6 +434,12 @@ public:
 	UPROPERTY()
 	class USimpleConstructionScript* SimpleConstructionScript;
 
+	UPROPERTY()
+	UStructProperty* UberGraphFramePointerProperty;
+
+	UPROPERTY()
+	UFunction* UberGraphFunction;
+
 	/** 
 	 * Gets an array of all BPGeneratedClasses (including InClass as 0th element) parents of given generated class 
 	 *
@@ -432,8 +466,19 @@ public:
 	virtual void ConditionalRecompileClass(TArray<UObject*>* ObjLoaded) override;
 #endif //WITH_EDITOR
 	virtual bool IsFunctionImplementedInBlueprint(FName InFunctionName) const override;
+	virtual uint8* GetPersistentUberGraphFrame(UObject* Obj, UFunction* FuncToCheck) const override;
+	virtual void CreatePersistentUberGraphFrame(UObject* Obj) const override;
+	virtual void DestroyPersistentUberGraphFrame(UObject* Obj) const override;
+	virtual void Link(FArchive& Ar, bool bRelinkExistingProperties) override;
+	virtual void PurgeClass(bool bRecompilingOnLoad) override;
+	virtual void Bind() override;
 	// End UClass interface
 
+	static void AddReferencedObjectsInUbergraphFrame(UObject* InThis, FReferenceCollector& Collector);
+
+	static FName GetUberGraphFrameName();
+	static bool UsePersistentUberGraphFrame();
+
 #if WITH_EDITORONLY_DATA
 	FBlueprintDebugData DebugData;
 
diff --git a/Engine/Source/Runtime/Engine/Private/BlueprintGeneratedClass.cpp b/Engine/Source/Runtime/Engine/Private/BlueprintGeneratedClass.cpp
--- a/Engine/Source/Runtime/Engine/Private/BlueprintGeneratedClass.cpp
+++ b/Engine/Source/Runtime/Engine/Private/BlueprintGeneratedClass.cpp
@@ -401,3 +401,140 @@ void UBlueprintGeneratedClass::CreateComponentsForActor(AActor* Actor) const
 		}
 	}
 }
+
+uint8* UBlueprintGeneratedClass::GetPersistentUberGraphFrame(UObject* Obj, UFunction* FuncToCheck) const
+{
+	if (Obj && UsePersistentUberGraphFrame() && UberGraphFramePointerProperty && UberGraphFunction)
+	{
+		if (UberGraphFunction == FuncToCheck)
+		{
+			auto PointerToUberGraphFrame = UberGraphFramePointerProperty->ContainerPtrToValuePtr<FPointerToUberGraphFrame>(Obj);
+			checkSlow(PointerToUberGraphFrame && PointerToUberGraphFrame->RawPointer);
+			return PointerToUberGraphFrame->RawPointer;
+		}
+	}
+	auto ParentClass = GetSuperClass();
+	checkSlow(ParentClass);
+	return ParentClass->GetPersistentUberGraphFrame(Obj, FuncToCheck);
+}
+
+void UBlueprintGeneratedClass::CreatePersistentUberGraphFrame(UObject* Obj) const
+{
+	checkSlow(!UberGraphFramePointerProperty == !UberGraphFunction);
+	if (Obj && UsePersistentUberGraphFrame() && UberGraphFramePointerProperty && UberGraphFunction)
+	{
+		auto FrameMemory = (uint8*)FMemory::Malloc(UberGraphFunction->GetStructureSize());
+		FMemory::Memzero(FrameMemory, UberGraphFunction->GetStructureSize());
+		for (UProperty* Property = UberGraphFunction->PropertyLink; Property; Property = Property->PropertyLinkNext)
+		{
+			Property->InitializeValue_InContainer(FrameMemory);
+		}
+		
+		auto PointerToUberGraphFrame = UberGraphFramePointerProperty->ContainerPtrToValuePtr<FPointerToUberGraphFrame>(Obj);
+		checkSlow(PointerToUberGraphFrame && !PointerToUberGraphFrame->RawPointer);
+		PointerToUberGraphFrame->RawPointer = FrameMemory;
+	}
+
+	auto ParentClass = GetSuperClass();
+	checkSlow(ParentClass);
+	return ParentClass->CreatePersistentUberGraphFrame(Obj);
+}
+
+void UBlueprintGeneratedClass::DestroyPersistentUberGraphFrame(UObject* Obj) const
+{
+	checkSlow(!UberGraphFramePointerProperty == !UberGraphFunction);
+	if (Obj && UsePersistentUberGraphFrame() && UberGraphFramePointerProperty && UberGraphFunction)
+	{
+		auto PointerToUberGraphFrame = UberGraphFramePointerProperty->ContainerPtrToValuePtr<FPointerToUberGraphFrame>(Obj);
+		checkSlow(PointerToUberGraphFrame && PointerToUberGraphFrame->RawPointer);
+		auto FrameMemory = PointerToUberGraphFrame->RawPointer;
+		PointerToUberGraphFrame->RawPointer = NULL;
+
+		for (UProperty* Property = UberGraphFunction->PropertyLink; Property; Property = Property->PropertyLinkNext)
+		{
+			Property->DestroyValue_InContainer(FrameMemory);
+		}
+		FMemory::Free(FrameMemory);
+	}
+
+	auto ParentClass = GetSuperClass();
+	checkSlow(ParentClass);
+	return ParentClass->DestroyPersistentUberGraphFrame(Obj);
+}
+
+void UBlueprintGeneratedClass::Link(FArchive& Ar, bool bRelinkExistingProperties)
+{
+	Super::Link(Ar, bRelinkExistingProperties);
+
+	if (UsePersistentUberGraphFrame() && UberGraphFunction)
+	{
+		for (auto Property : TFieldRange<UStructProperty>(this, EFieldIteratorFlags::ExcludeSuper))
+		{
+			if (Property->GetFName() == GetUberGraphFrameName())
+			{
+				UberGraphFramePointerProperty = Property;
+				break;
+			}
+		}
+		checkSlow(UberGraphFramePointerProperty);
+	}
+}
+
+void UBlueprintGeneratedClass::PurgeClass(bool bRecompilingOnLoad)
+{
+	Super::PurgeClass(bRecompilingOnLoad);
+
+	UberGraphFramePointerProperty = NULL;
+	UberGraphFunction = NULL;
+}
+
+void UBlueprintGeneratedClass::Bind()
+{
+	Super::Bind();
+
+	if (UsePersistentUberGraphFrame() && UberGraphFunction)
+	{
+		ClassAddReferencedObjects = &UBlueprintGeneratedClass::AddReferencedObjectsInUbergraphFrame;
+	}
+}
+
+void UBlueprintGeneratedClass::AddReferencedObjectsInUbergraphFrame(UObject* InThis, FReferenceCollector& Collector)
+{
+	checkSlow(InThis);
+	for (UClass* CurrentClass = InThis->GetClass(); CurrentClass; CurrentClass = CurrentClass->GetSuperClass())
+	{
+		if (auto BPGC = Cast<UBlueprintGeneratedClass>(CurrentClass))
+		{
+			if (BPGC->UberGraphFramePointerProperty)
+			{
+				checkSlow(BPGC->UberGraphFunction);
+				auto PointerToUberGraphFrame = BPGC->UberGraphFramePointerProperty->ContainerPtrToValuePtr<FPointerToUberGraphFrame>(InThis);
+				checkSlow(PointerToUberGraphFrame && PointerToUberGraphFrame->RawPointer);
+
+				FSimpleObjectReferenceCollectorArchive ObjectReferenceCollector(InThis, Collector);
+				BPGC->UberGraphFunction->SerializeBin(ObjectReferenceCollector, PointerToUberGraphFrame->RawPointer, 0);
+			}
+		}
+		else if (CurrentClass->HasAllClassFlags(CLASS_Native))
+		{
+			CurrentClass->CallAddReferencedObjects(InThis, Collector);
+			break;
+		}
+		else
+		{
+			checkSlow(false);
+		}
+	}
+}
+
+FName UBlueprintGeneratedClass::GetUberGraphFrameName()
+{
+	static const FName UberGraphFrameName(TEXT("UberGraphFrame"));
+	return UberGraphFrameName;
+}
+
+bool UBlueprintGeneratedClass::UsePersistentUberGraphFrame()
+{
+	static const FBoolConfigValueHelper PersistentUberGraphFrame(TEXT("Kismet"), TEXT("bPersistentUberGraphFrame"), GEngineIni);
+	return PersistentUberGraphFrame;
+}
\ No newline at end of file
