Editor - Fixing a warning in the BP Compiler.

diff --git a/Engine/Source/Editor/KismetCompiler/Private/KismetCompiler.cpp b/Engine/Source/Editor/KismetCompiler/Private/KismetCompiler.cpp
--- a/Engine/Source/Editor/KismetCompiler/Private/KismetCompiler.cpp
+++ b/Engine/Source/Editor/KismetCompiler/Private/KismetCompiler.cpp
@@ -2444,8 +2444,8 @@ void FKismetCompilerContext::MergeUbergraphPagesIn(UEdGraph* Ubergraph)
 			FEdGraphUtilities::CloneAndMergeGraphIn(Ubergraph, SourceGraph, MessageLog, /*bRequireSchemaMatch=*/ true, &ClonedNodeList);
 
 			// Create a comment block around the ubergrapgh contents before anything else got started
-			int32 OffsetX;
-			int32 OffsetY;
+			int32 OffsetX = 0;
+			int32 OffsetY = 0;
 			CreateCommentBlockAroundNodes(ClonedNodeList, SourceGraph, Ubergraph, SourceGraph->GetName(), FLinearColor(1.0f, 0.7f, 0.7f), /*out*/ OffsetX, /*out*/ OffsetY);
 			
 			// Reposition the nodes, so nothing ever overlaps
