[UE-1987] Now allowing for users to cast from interfaces to raw objects.

diff --git a/Engine/Source/Editor/BlueprintGraph/Private/DynamicCastHandler.cpp b/Engine/Source/Editor/BlueprintGraph/Private/DynamicCastHandler.cpp
--- a/Engine/Source/Editor/BlueprintGraph/Private/DynamicCastHandler.cpp
+++ b/Engine/Source/Editor/BlueprintGraph/Private/DynamicCastHandler.cpp
@@ -96,13 +96,19 @@ void FKCHandler_DynamicCast::Compile(FKismetFunctionContext& Context, UEdGraphNo
 	UClass const* const OutputObjClass = Cast<UClass>((*CastResultTerm)->Type.PinSubCategoryObject.Get());
 
 	const bool bIsOutputInterface = ((OutputObjClass != NULL) && OutputObjClass->HasAnyClassFlags(CLASS_Interface));
-	const bool bIsInputInterface = ((InputObjClass != NULL) && InputObjClass->HasAnyClassFlags(CLASS_Interface));
+	const bool bIsInputInterface  = ((InputObjClass != NULL) && InputObjClass->HasAnyClassFlags(CLASS_Interface));
 
 	EKismetCompiledStatementType CastOpType = KCST_DynamicCast;
 	if (bIsInputInterface)
 	{
-		check(bIsOutputInterface);
-		CastOpType = KCST_CrossInterfaceCast;
+		if (bIsOutputInterface)
+		{
+			CastOpType = KCST_CrossInterfaceCast;
+		}
+		else
+		{
+			CastOpType = KCST_CastInterfaceToObj;
+		}
 	}
 	else if (bIsOutputInterface)
 	{
diff --git a/Engine/Source/Editor/KismetCompiler/Private/KismetCompilerBackend.h b/Engine/Source/Editor/KismetCompiler/Private/KismetCompilerBackend.h
--- a/Engine/Source/Editor/KismetCompiler/Private/KismetCompilerBackend.h
+++ b/Engine/Source/Editor/KismetCompiler/Private/KismetCompilerBackend.h
@@ -101,6 +101,7 @@ public:
 	void EmitAssignmentStatment(FKismetFunctionContext& FunctionContext, FBlueprintCompiledStatement& Statement);
 	void EmitCastObjToInterfaceStatement(FKismetFunctionContext& FunctionContext, FBlueprintCompiledStatement& Statement);
 	void EmitCastBetweenInterfacesStatement(FKismetFunctionContext& FunctionContext, FBlueprintCompiledStatement& Statement);
+	void EmitCastInterfaceToObjStatement(FKismetFunctionContext& FunctionContext, FBlueprintCompiledStatement& Statement);
 	void EmitDynamicCastStatement(FKismetFunctionContext& FunctionContext, FBlueprintCompiledStatement& Statement);
 	void EmitObjectToBoolStatement(FKismetFunctionContext& FunctionContext, FBlueprintCompiledStatement& Statement);
 	void EmitAddMulticastDelegateStatement(FKismetFunctionContext& FunctionContext, FBlueprintCompiledStatement& Statement);
diff --git a/Engine/Source/Editor/KismetCompiler/Private/KismetCompilerCppBackend.cpp b/Engine/Source/Editor/KismetCompiler/Private/KismetCompilerCppBackend.cpp
--- a/Engine/Source/Editor/KismetCompiler/Private/KismetCompilerCppBackend.cpp
+++ b/Engine/Source/Editor/KismetCompiler/Private/KismetCompilerCppBackend.cpp
@@ -879,6 +879,18 @@ void FKismetCppBackend::EmitCastBetweenInterfacesStatement(FKismetFunctionContex
 	Emit(Body, *FString::Printf(TEXT("\t\t\t}\n")));
 }
 
+void FKismetCppBackend::EmitCastInterfaceToObjStatement(FKismetFunctionContext& FunctionContext, FBlueprintCompiledStatement& Statement)
+{
+	FString ClassToCastTo = TermToText(Statement.RHS[0], (UProperty*)(GetDefault<UClassProperty>()));
+	FString InputInterface = TermToText(Statement.RHS[1], (UProperty*)(GetDefault<UInterfaceProperty>()));
+	FString ResultObject = TermToText(Statement.LHS, (UProperty*)(GetDefault<UInterfaceProperty>()));
+
+	FString InputObject = FString::Printf(TEXT("%s.GetObjectRef()"), *InputInterface);
+	
+	Emit(Body, *FString::Printf(TEXT("\t\t\t%s = Cast<%s>(%s);\n"),
+		*ResultObject, *ClassToCastTo, *InputObject));
+}
+
 void FKismetCppBackend::EmitDynamicCastStatement(FKismetFunctionContext& FunctionContext, FBlueprintCompiledStatement& Statement)
 {
 	FString TargetClass = TermToText(Statement.RHS[0], (UProperty*)(GetDefault<UClassProperty>()));
@@ -1252,6 +1264,9 @@ void FKismetCppBackend::ConstructFunction(FKismetFunctionContext& FunctionContex
 					case KCST_CrossInterfaceCast:
 						EmitCastBetweenInterfacesStatement(FunctionContext, Statement);
 						break;
+					case KCST_CastInterfaceToObj:
+						EmitCastInterfaceToObjStatement(FunctionContext, Statement);
+						break;
 					case KCST_DynamicCast:
 						EmitDynamicCastStatement(FunctionContext, Statement);
 						break;
diff --git a/Engine/Source/Editor/KismetCompiler/Private/KismetCompilerVMBackend.cpp b/Engine/Source/Editor/KismetCompiler/Private/KismetCompilerVMBackend.cpp
--- a/Engine/Source/Editor/KismetCompiler/Private/KismetCompilerVMBackend.cpp
+++ b/Engine/Source/Editor/KismetCompiler/Private/KismetCompilerVMBackend.cpp
@@ -893,17 +893,33 @@ public:
 		EmitTerm(TargetExpression, (UProperty*)(GetDefault<UInterfaceProperty>()));
 	}
 
+	void EmitCastInterfaceToObjStatement(FBlueprintCompiledStatement& Statement)
+	{
+		FBPTerminal* DestinationExpression     = Statement.LHS;
+		FBPTerminal* ResultObjClassExpression  = Statement.RHS[0];
+		FBPTerminal* TargetInterfaceExpression = Statement.RHS[1];
+
+		Writer << EX_Let;
+		EmitTerm(DestinationExpression);
+
+		Writer << EX_InterfaceToObjCast;
+		UClass* ClassPtr = CastChecked<UClass>(ResultObjClassExpression->ObjectLiteral);
+		check(ClassPtr != nullptr);
+		Writer << ClassPtr;
+		EmitTerm(TargetInterfaceExpression, (UProperty*)(GetDefault<UObjectProperty>()));
+	}
+
 	void EmitDynamicCastStatement(FBlueprintCompiledStatement& Statement)
 	{
 		FBPTerminal* DestinationExpression = Statement.LHS;
-		FBPTerminal* InterfaceExpression = Statement.RHS[0];
+		FBPTerminal* ResultClassExpression = Statement.RHS[0];
 		FBPTerminal* TargetExpression = Statement.RHS[1];
 
 		Writer << EX_Let;
 		EmitTerm(DestinationExpression);
 
 		Writer << EX_DynamicCast;
-		UClass* ClassPtr = CastChecked<UClass>(InterfaceExpression->ObjectLiteral);
+		UClass* ClassPtr = CastChecked<UClass>(ResultClassExpression->ObjectLiteral);
 		Writer << ClassPtr;
 		EmitTerm(TargetExpression, (UProperty*)(GetDefault<UObjectProperty>()));
 	}
@@ -1182,6 +1198,9 @@ public:
 		case KCST_CrossInterfaceCast:
 			EmitCastBetweenInterfacesStatement(Statement);
 			break;
+		case KCST_CastInterfaceToObj:
+			EmitCastInterfaceToObjStatement(Statement);
+			break;
 		case KCST_DynamicCast:
 			EmitDynamicCastStatement(Statement);
 			break;
diff --git a/Engine/Source/Editor/KismetCompiler/Public/BlueprintCompiledStatement.h b/Engine/Source/Editor/KismetCompiler/Public/BlueprintCompiledStatement.h
--- a/Engine/Source/Editor/KismetCompiler/Public/BlueprintCompiledStatement.h
+++ b/Engine/Source/Editor/KismetCompiler/Public/BlueprintCompiledStatement.h
@@ -33,6 +33,7 @@ enum EKismetCompiledStatementType
 	KCST_CrossInterfaceCast = 23,	// TargetInterface(Interface)
 	KCST_MetaCast = 24,	    // Cast<TargetClass>(TargetObject)
 	KCST_AssignmentOnPersistentFrame = 25, //
+	KCST_CastInterfaceToObj = 26, // Cast<TargetClass>(TargetInterface)
 };
 
 //@TODO: Too rigid / icky design
diff --git a/Engine/Source/Editor/UnrealEd/Private/ScriptDisassembler.cpp b/Engine/Source/Editor/UnrealEd/Private/ScriptDisassembler.cpp
--- a/Engine/Source/Editor/UnrealEd/Private/ScriptDisassembler.cpp
+++ b/Engine/Source/Editor/UnrealEd/Private/ScriptDisassembler.cpp
@@ -218,6 +218,18 @@ void FKismetBytecodeDisassembler::ProcessCommon(int32& ScriptIndex, EExprToken O
 			UClass* InterfaceClass = ReadPointer<UClass>(ScriptIndex);
 			Ar.Logf(TEXT("%s $%X: InterfaceToInterfaceCast to %s"), *Indents, (int32)Opcode, *InterfaceClass->GetName());
 
+			SerializeExpr( ScriptIndex );
+			break;
+		}
+	case EX_InterfaceToObjCast:
+		{
+			// A conversion from an interface variable to a object variable.
+			// We use a different bytecode to avoid the branching each time we process a cast token
+
+			// the interface class to convert to
+			UClass* ObjectClass = ReadPointer<UClass>(ScriptIndex);
+			Ar.Logf(TEXT("%s $%X: InterfaceToObjCast to %s"), *Indents, (int32)Opcode, *ObjectClass->GetName());
+
 			SerializeExpr( ScriptIndex );
 			break;
 		}
diff --git a/Engine/Source/Runtime/CoreUObject/Private/UObject/ScriptCore.cpp b/Engine/Source/Runtime/CoreUObject/Private/UObject/ScriptCore.cpp
--- a/Engine/Source/Runtime/CoreUObject/Private/UObject/ScriptCore.cpp
+++ b/Engine/Source/Runtime/CoreUObject/Private/UObject/ScriptCore.cpp
@@ -2008,3 +2008,25 @@ void UObject::execInterfaceToInterface( FFrame& Stack, RESULT_DECL )
 }
 IMPLEMENT_VM_FUNCTION( EX_CrossInterfaceCast, execInterfaceToInterface );
 
+void UObject::execInterfaceToObject(FFrame& Stack, RESULT_DECL)
+{
+	// read the interface class off the stack
+	UClass* ObjClassToCastTo = dynamic_cast<UClass*>(Stack.ReadObject());
+	checkSlow(ObjClassToCastTo != nullptr);
+
+	// read the input interface-object off the stack
+	FScriptInterface InterfaceInput;
+	Stack.Step(Stack.Object, &InterfaceInput);
+
+	UObject* InputObjWithInterface = InterfaceInput.GetObjectRef();
+	if (InputObjWithInterface->IsA(ObjClassToCastTo))
+	{
+		*(UObject**)Result = InputObjWithInterface;
+	}
+	else
+	{
+		*(UObject**)Result = nullptr;
+	}
+}
+IMPLEMENT_VM_FUNCTION( EX_InterfaceToObjCast, execInterfaceToObject );
+
diff --git a/Engine/Source/Runtime/CoreUObject/Public/UObject/Script.h b/Engine/Source/Runtime/CoreUObject/Public/UObject/Script.h
--- a/Engine/Source/Runtime/CoreUObject/Public/UObject/Script.h
+++ b/Engine/Source/Runtime/CoreUObject/Public/UObject/Script.h
@@ -219,7 +219,7 @@ enum EExprToken
 	EX_ObjToInterfaceCast   = 0x52,	// Converting an object reference to native interface variable
 	EX_EndOfScript			= 0x53, // Last byte in script code
 	EX_CrossInterfaceCast	= 0x54, // Converting an interface variable reference to native interface variable
-	//						= 0x55,
+	EX_InterfaceToObjCast   = 0x55, // Converting an interface variable reference to an object
 	//						= 0x56,
 	//						= 0x57,
 	//						= 0x58,
diff --git a/Engine/Source/Runtime/CoreUObject/Public/UObject/ScriptSerialization.h b/Engine/Source/Runtime/CoreUObject/Public/UObject/ScriptSerialization.h
--- a/Engine/Source/Runtime/CoreUObject/Public/UObject/ScriptSerialization.h
+++ b/Engine/Source/Runtime/CoreUObject/Public/UObject/ScriptSerialization.h
@@ -110,6 +110,7 @@
 		}
 		case EX_ObjToInterfaceCast:
 		case EX_CrossInterfaceCast:
+		case EX_InterfaceToObjCast:
 		{
 			// A conversion from an object or interface variable to a native interface variable.  
 			// We use a different bytecode to avoid the branching each time we process a cast token.
diff --git a/Engine/Source/Runtime/CoreUObject/Public/UObject/UObject.h b/Engine/Source/Runtime/CoreUObject/Public/UObject/UObject.h
--- a/Engine/Source/Runtime/CoreUObject/Public/UObject/UObject.h
+++ b/Engine/Source/Runtime/CoreUObject/Public/UObject/UObject.h
@@ -951,6 +951,7 @@ public:
 	DECLARE_FUNCTION(execInterfaceToBool);
 	DECLARE_FUNCTION(execObjectToInterface);
 	DECLARE_FUNCTION(execInterfaceToInterface);
+	DECLARE_FUNCTION(execInterfaceToObject);
 
 	// Dynamic array functions
 	// Array support
