KismetCompilerCppBackend refactor - first step

diff --git a/Engine/Source/Editor/KismetCompiler/Private/KismetCompilerBackend.h b/Engine/Source/Editor/KismetCompiler/Private/KismetCompilerBackend.h
--- a/Engine/Source/Editor/KismetCompiler/Private/KismetCompilerBackend.h
+++ b/Engine/Source/Editor/KismetCompiler/Private/KismetCompilerBackend.h
@@ -18,112 +18,14 @@ class IKismetCompilerBackend
 //////////////////////////////////////////////////////////////////////////
 // FKismetCppBackend
 
-class FKismetCppBackend : public IKismetCompilerBackend
+class IKismetCppBackend : public IKismetCompilerBackend
 {
-protected:
-	UEdGraphSchema_K2* Schema;
-	FCompilerResultsLog& MessageLog;
-	FKismetCompilerContext& CompilerContext;
-
-	struct FFunctionLabelInfo
-	{
-		TMap<FBlueprintCompiledStatement*, int32> StateMap;
-		int32 StateCounter;
-
-		FFunctionLabelInfo()
-		{
-			StateCounter = 0;
-		}
-
-		int32 StatementToStateIndex(FBlueprintCompiledStatement* Statement)
-		{
-			int32& Index = StateMap.FindOrAdd(Statement);
-			if (Index == 0)
-			{
-				Index = ++StateCounter;
-			}
-
-			return Index;
-		}
-	};
-	
-	TArray<FFunctionLabelInfo> StateMapPerFunction;
-	TMap<FKismetFunctionContext*, int32> FunctionIndexMap; 
-
-	FString CppClassName;
-
-	// Pointers to commonly used structures (found in constructor)
-	UScriptStruct* VectorStruct;
-	UScriptStruct* RotatorStruct;
-	UScriptStruct* TransformStruct;
-	UScriptStruct* LatentInfoStruct;
-	UScriptStruct* LinearColorStruct;
 public:
-	FStringOutputDevice Header;
-	FStringOutputDevice Body;
-	FString TermToText(const FBPTerminal* Term, const UProperty* SourceProperty = NULL);
-protected:
-	FString LatentFunctionInfoTermToText(FBPTerminal* Term, FBlueprintCompiledStatement* TargetLabel);
+	static IKismetCppBackend* Create(UEdGraphSchema_K2* InSchema, FKismetCompilerContext& InContext);
 
-	int32 StatementToStateIndex(FKismetFunctionContext& FunctionContext, FBlueprintCompiledStatement* Statement)
-	{
-		int32 Index = FunctionIndexMap.FindChecked(&FunctionContext);
-		return StateMapPerFunction[Index].StatementToStateIndex(Statement);
-	}
-public:
-
-	FKismetCppBackend(UEdGraphSchema_K2* InSchema, FKismetCompilerContext& InContext)
-		: Schema(InSchema)
-		, MessageLog(InContext.MessageLog)
-		, CompilerContext(InContext)
-	{
-		extern UScriptStruct* Z_Construct_UScriptStruct_UObject_FVector();
-		VectorStruct = Z_Construct_UScriptStruct_UObject_FVector();
-		RotatorStruct = FindObjectChecked<UScriptStruct>(UObject::StaticClass(), TEXT("Rotator"));
-		TransformStruct = FindObjectChecked<UScriptStruct>(UObject::StaticClass(), TEXT("Transform"));
-		LinearColorStruct = FindObjectChecked<UScriptStruct>(UObject::StaticClass(), TEXT("LinearColor"));
-		LatentInfoStruct = FLatentActionInfo::StaticStruct();
-	}
-
-	UEdGraphSchema_K2* GetSchema() const { return Schema; }
-
-	void Emit(FStringOutputDevice& Target, const TCHAR* Message)
-	{
-		Target += Message;
-	}
-
-	void EmitClassProperties(FStringOutputDevice& Target, UClass* SourceClass);
-
-	void GenerateCodeFromClass(UClass* SourceClass, TIndirectArray<FKismetFunctionContext>& Functions, bool bGenerateStubsOnly=false);
-
-	void EmitCallStatment(FKismetFunctionContext& FunctionContext, FBlueprintCompiledStatement& Statement);
-	void EmitCallDelegateStatment(FKismetFunctionContext& FunctionContext, FBlueprintCompiledStatement& Statement);
-	void EmitAssignmentStatment(FKismetFunctionContext& FunctionContext, FBlueprintCompiledStatement& Statement);
-	void EmitCastObjToInterfaceStatement(FKismetFunctionContext& FunctionContext, FBlueprintCompiledStatement& Statement);
-	void EmitCastBetweenInterfacesStatement(FKismetFunctionContext& FunctionContext, FBlueprintCompiledStatement& Statement);
-	void EmitCastInterfaceToObjStatement(FKismetFunctionContext& FunctionContext, FBlueprintCompiledStatement& Statement);
-	void EmitDynamicCastStatement(FKismetFunctionContext& FunctionContext, FBlueprintCompiledStatement& Statement);
-	void EmitObjectToBoolStatement(FKismetFunctionContext& FunctionContext, FBlueprintCompiledStatement& Statement);
-	void EmitAddMulticastDelegateStatement(FKismetFunctionContext& FunctionContext, FBlueprintCompiledStatement& Statement);
-	void EmitBindDelegateStatement(FKismetFunctionContext& FunctionContext, FBlueprintCompiledStatement& Statement);
-	void EmitClearMulticastDelegateStatement(FKismetFunctionContext& FunctionContext, FBlueprintCompiledStatement& Statement);
-	void EmitCreateArrayStatement(FKismetFunctionContext& FunctionContext, FBlueprintCompiledStatement& Statement);
-	void EmitGotoStatement(FKismetFunctionContext& FunctionContext, FBlueprintCompiledStatement& Statement);
-	void EmitPushStateStatement(FKismetFunctionContext& FunctionContext, FBlueprintCompiledStatement& Statement);
-	void EmitEndOfThreadStatement(FKismetFunctionContext& FunctionContext, const FString& ReturnValueString);
-	void EmitReturnStatement(FKismetFunctionContext& FunctionContext, const FString& ReturnValueString);
-	void EmitRemoveMulticastDelegateStatement(FKismetFunctionContext& FunctionContext, FBlueprintCompiledStatement& Statement);
-	void EmitMetaCastStatement(FKismetFunctionContext& FunctionContext, FBlueprintCompiledStatement& Statement);
-	
-	/** Emits local variable declarations for a function */
-	void DeclareLocalVariables(FKismetFunctionContext& FunctionContext, TArray<UProperty*>& LocalVariables);
-
-	/** Emits the internal execution flow state declaration for a function */
-	void DeclareStateSwitch(FKismetFunctionContext& FunctionContext);
-	void CloseStateSwitch(FKismetFunctionContext& FunctionContext);
-
-	/** Builds both the header declaration and body implementation of a function */
-	void ConstructFunction(FKismetFunctionContext& FunctionContext, bool bGenerateStubOnly);
+	virtual const FString& GetBody() const = 0;
+	virtual const FString& GetHeader() const = 0;
+	virtual void GenerateCodeFromClass(UClass* SourceClass, TIndirectArray<FKismetFunctionContext>& Functions, bool bGenerateStubsOnly = false) = 0;
 };
 
 //////////////////////////////////////////////////////////////////////////
diff --git a/Engine/Source/Editor/KismetCompiler/Private/KismetCompilerCppBackend.cpp b/Engine/Source/Editor/KismetCompiler/Private/KismetCompilerCppBackend.cpp
--- a/Engine/Source/Editor/KismetCompiler/Private/KismetCompilerCppBackend.cpp
+++ b/Engine/Source/Editor/KismetCompiler/Private/KismetCompilerCppBackend.cpp
@@ -10,6 +10,126 @@
 
 #include "DefaultValueHelper.h"
 
+//////////////////////////////////////////////////////////////////////////
+// FKismetCppBackend
+
+class FKismetCppBackend : public IKismetCppBackend
+{
+protected:
+	UEdGraphSchema_K2* Schema;
+	FCompilerResultsLog& MessageLog;
+	FKismetCompilerContext& CompilerContext;
+
+	struct FFunctionLabelInfo
+	{
+		TMap<FBlueprintCompiledStatement*, int32> StateMap;
+		int32 StateCounter;
+
+		FFunctionLabelInfo()
+		{
+			StateCounter = 0;
+		}
+
+		int32 StatementToStateIndex(FBlueprintCompiledStatement* Statement)
+		{
+			int32& Index = StateMap.FindOrAdd(Statement);
+			if (Index == 0)
+			{
+				Index = ++StateCounter;
+			}
+
+			return Index;
+		}
+	};
+
+	TArray<FFunctionLabelInfo> StateMapPerFunction;
+	TMap<FKismetFunctionContext*, int32> FunctionIndexMap;
+
+	FString CppClassName;
+
+	// Pointers to commonly used structures (found in constructor)
+	UScriptStruct* VectorStruct;
+	UScriptStruct* RotatorStruct;
+	UScriptStruct* TransformStruct;
+	UScriptStruct* LatentInfoStruct;
+	UScriptStruct* LinearColorStruct;
+public:
+	FStringOutputDevice Header;
+	FStringOutputDevice Body;
+	FString TermToText(const FBPTerminal* Term, const UProperty* SourceProperty = NULL);
+
+	const FString& GetBody()	const override { return Body; }
+	const FString& GetHeader()	const override { return Header; }
+
+protected:
+	FString LatentFunctionInfoTermToText(FBPTerminal* Term, FBlueprintCompiledStatement* TargetLabel);
+
+	int32 StatementToStateIndex(FKismetFunctionContext& FunctionContext, FBlueprintCompiledStatement* Statement)
+	{
+		int32 Index = FunctionIndexMap.FindChecked(&FunctionContext);
+		return StateMapPerFunction[Index].StatementToStateIndex(Statement);
+	}
+public:
+
+	FKismetCppBackend(UEdGraphSchema_K2* InSchema, FKismetCompilerContext& InContext)
+		: Schema(InSchema)
+		, MessageLog(InContext.MessageLog)
+		, CompilerContext(InContext)
+	{
+		extern UScriptStruct* Z_Construct_UScriptStruct_UObject_FVector();
+		VectorStruct = Z_Construct_UScriptStruct_UObject_FVector();
+		RotatorStruct = FindObjectChecked<UScriptStruct>(UObject::StaticClass(), TEXT("Rotator"));
+		TransformStruct = FindObjectChecked<UScriptStruct>(UObject::StaticClass(), TEXT("Transform"));
+		LinearColorStruct = FindObjectChecked<UScriptStruct>(UObject::StaticClass(), TEXT("LinearColor"));
+		LatentInfoStruct = FLatentActionInfo::StaticStruct();
+	}
+
+	UEdGraphSchema_K2* GetSchema() const { return Schema; }
+
+	void Emit(FStringOutputDevice& Target, const TCHAR* Message)
+	{
+		Target += Message;
+	}
+
+	void EmitClassProperties(FStringOutputDevice& Target, UClass* SourceClass);
+
+	void GenerateCodeFromClass(UClass* SourceClass, TIndirectArray<FKismetFunctionContext>& Functions, bool bGenerateStubsOnly) override;
+
+	void EmitCallStatment(FKismetFunctionContext& FunctionContext, FBlueprintCompiledStatement& Statement);
+	void EmitCallDelegateStatment(FKismetFunctionContext& FunctionContext, FBlueprintCompiledStatement& Statement);
+	void EmitAssignmentStatment(FKismetFunctionContext& FunctionContext, FBlueprintCompiledStatement& Statement);
+	void EmitCastObjToInterfaceStatement(FKismetFunctionContext& FunctionContext, FBlueprintCompiledStatement& Statement);
+	void EmitCastBetweenInterfacesStatement(FKismetFunctionContext& FunctionContext, FBlueprintCompiledStatement& Statement);
+	void EmitCastInterfaceToObjStatement(FKismetFunctionContext& FunctionContext, FBlueprintCompiledStatement& Statement);
+	void EmitDynamicCastStatement(FKismetFunctionContext& FunctionContext, FBlueprintCompiledStatement& Statement);
+	void EmitObjectToBoolStatement(FKismetFunctionContext& FunctionContext, FBlueprintCompiledStatement& Statement);
+	void EmitAddMulticastDelegateStatement(FKismetFunctionContext& FunctionContext, FBlueprintCompiledStatement& Statement);
+	void EmitBindDelegateStatement(FKismetFunctionContext& FunctionContext, FBlueprintCompiledStatement& Statement);
+	void EmitClearMulticastDelegateStatement(FKismetFunctionContext& FunctionContext, FBlueprintCompiledStatement& Statement);
+	void EmitCreateArrayStatement(FKismetFunctionContext& FunctionContext, FBlueprintCompiledStatement& Statement);
+	void EmitGotoStatement(FKismetFunctionContext& FunctionContext, FBlueprintCompiledStatement& Statement);
+	void EmitPushStateStatement(FKismetFunctionContext& FunctionContext, FBlueprintCompiledStatement& Statement);
+	void EmitEndOfThreadStatement(FKismetFunctionContext& FunctionContext, const FString& ReturnValueString);
+	void EmitReturnStatement(FKismetFunctionContext& FunctionContext, const FString& ReturnValueString);
+	void EmitRemoveMulticastDelegateStatement(FKismetFunctionContext& FunctionContext, FBlueprintCompiledStatement& Statement);
+	void EmitMetaCastStatement(FKismetFunctionContext& FunctionContext, FBlueprintCompiledStatement& Statement);
+
+	/** Emits local variable declarations for a function */
+	void DeclareLocalVariables(FKismetFunctionContext& FunctionContext, TArray<UProperty*>& LocalVariables);
+
+	/** Emits the internal execution flow state declaration for a function */
+	void DeclareStateSwitch(FKismetFunctionContext& FunctionContext);
+	void CloseStateSwitch(FKismetFunctionContext& FunctionContext);
+
+	/** Builds both the header declaration and body implementation of a function */
+	void ConstructFunction(FKismetFunctionContext& FunctionContext, bool bGenerateStubOnly);
+};
+
+IKismetCppBackend* IKismetCppBackend::Create(UEdGraphSchema_K2* InSchema, FKismetCompilerContext& InContext)
+{
+	return new FKismetCppBackend(InSchema, InContext);
+}
+
 // Generates single if scope. It's condition checks context of given term.
 struct FSafeContextScopedEmmitter
 {
