BP: Public functionality to generate cpp code from blueprint

diff --git a/Engine/Source/Editor/Kismet/Private/NativeCodeGenerationTool.cpp b/Engine/Source/Editor/Kismet/Private/NativeCodeGenerationTool.cpp
--- a/Engine/Source/Editor/Kismet/Private/NativeCodeGenerationTool.cpp
+++ b/Engine/Source/Editor/Kismet/Private/NativeCodeGenerationTool.cpp
@@ -8,18 +8,20 @@
 
 struct FGeneratedCodeData
 {
-	FGeneratedCodeData(UBlueprint& Blueprint)
+	FGeneratedCodeData(UBlueprint& Blueprint) : HeaderSource(new FString()), CppSource(new FString())
 	{
 		FName GeneratedClassName, SkeletonClassName;
 		Blueprint.GetBlueprintClassNames(GeneratedClassName, SkeletonClassName);
 		ClassName = GeneratedClassName.ToString();
 
 		GatherUserDefinedDependencies(Blueprint);
+		
+		FKismetEditorUtilities::GenerateCppCode(&Blueprint, HeaderSource, CppSource);
 	}
 
 	FString TypeDependencies;
-	FString HeaderSource;
-	FString CppSource;
+	TSharedPtr<FString> HeaderSource;
+	TSharedPtr<FString> CppSource;
 	FString ErrorString;
 	FString ClassName;
 
@@ -90,12 +92,12 @@ struct FGeneratedCodeData
 
 	bool Save(const FString& HeaderDirPath, const FString& CppDirPath)
 	{
-		const bool bHeaderSaved = FFileHelper::SaveStringToFile(HeaderSource, *FPaths::Combine(*HeaderDirPath, *HeaderFileName()));
+		const bool bHeaderSaved = FFileHelper::SaveStringToFile(*HeaderSource, *FPaths::Combine(*HeaderDirPath, *HeaderFileName()));
 		if (!bHeaderSaved)
 		{
 			ErrorString += LOCTEXT("HeaderNotSaved", "Header file wasn't saved. Check log for details.\n").ToString();
 		}
-		const bool bCppSaved = FFileHelper::SaveStringToFile(CppSource, *FPaths::Combine(*CppDirPath, *SourceFileName()));
+		const bool bCppSaved = FFileHelper::SaveStringToFile(*CppSource, *FPaths::Combine(*CppDirPath, *SourceFileName()));
 		if (!bCppSaved)
 		{
 			ErrorString += LOCTEXT("CppNotSaved", "Cpp file wasn't saved. Check log for details.\n").ToString();
@@ -119,6 +121,8 @@ public:
 	FString Directory;
 	FText Message;
 
+	TSharedPtr<SEditableTextBox> EditableTextBox;
+
 	FString GetFilePath() const
 	{
 		return FPaths::Combine(*Directory, *File);
@@ -132,6 +136,10 @@ public:
 	FReply BrowseHeaderDirectory()
 	{
 		PromptUserForDirectory(Directory, Message.ToString(), Directory);
+
+		//workaround for UI problem
+		EditableTextBox->SetText(TAttribute<FText>::Create(TAttribute<FText>::FGetter::CreateSP(this, &SSimpleDirectoryPicker::GetFilePathText)));
+
 		return FReply::Handled();
 	}
 
@@ -149,7 +157,7 @@ public:
 			.FillWidth(1.0f)
 			.Padding(4.0f, 0.0f, 0.0f, 0.0f)
 			[
-				SNew(SEditableTextBox)
+				SAssignNew(EditableTextBox, SEditableTextBox)
 				.Text(this, &SSimpleDirectoryPicker::GetFilePathText)
 				.IsReadOnly(true)
 			]
diff --git a/Engine/Source/Editor/KismetCompiler/Private/KismetCompiler.cpp b/Engine/Source/Editor/KismetCompiler/Private/KismetCompiler.cpp
--- a/Engine/Source/Editor/KismetCompiler/Private/KismetCompiler.cpp
+++ b/Engine/Source/Editor/KismetCompiler/Private/KismetCompiler.cpp
@@ -702,7 +702,7 @@ void FKismetCompilerContext::CreatePropertiesFromList(UStruct* Scope, UField**&
 void FKismetCompilerContext::CreateLocalVariablesForFunction(FKismetFunctionContext& Context)
 {
 	ensure(Context.IsEventGraph() || !Context.EventGraphLocals.Num());
-	ensure(!Context.IsEventGraph() || !Context.Locals.Num());
+	ensure(!Context.IsEventGraph() || !Context.Locals.Num() || !UsePersistentUberGraphFrame());
 
 	const bool bPersistentUberGraphFrame = UsePersistentUberGraphFrame() && Context.bIsUbergraph;
 	// Local stack frame (or maybe class for the ubergraph)
@@ -2296,7 +2296,7 @@ void FKismetCompilerContext::CreateFunctionStubForEvent(UK2Node_Event* SrcEventN
 	ChildStubGraph->SetFlags(RF_Transient);
 	MessageLog.NotifyIntermediateObjectCreation(ChildStubGraph, SrcEventNode);
 
-	FKismetFunctionContext& StubContext = *new (FunctionList) FKismetFunctionContext(MessageLog, Schema, NewClass, Blueprint);
+	FKismetFunctionContext& StubContext = *new (FunctionList) FKismetFunctionContext(MessageLog, Schema, NewClass, Blueprint, CompileOptions.DoesRequireCppCodeGeneration());
 	StubContext.SourceGraph = ChildStubGraph;
 
 	// A stub graph has no visual representation and is thus not suited to be debugged via the debugger
@@ -2644,7 +2644,7 @@ void FKismetCompilerContext::CreateAndProcessUbergraph()
 
 		// Do some cursory validation (pin types match, inputs to outputs, pins never point to their parent node, etc...)
 		{
-			UbergraphContext = new (FunctionList)FKismetFunctionContext(MessageLog, Schema, NewClass, Blueprint);
+			UbergraphContext = new (FunctionList)FKismetFunctionContext(MessageLog, Schema, NewClass, Blueprint, CompileOptions.DoesRequireCppCodeGeneration());
 			UbergraphContext->SourceGraph = ConsolidatedEventGraph;
 			UbergraphContext->MarkAsEventGraph();
 			UbergraphContext->MarkAsInternalOrCppUseOnly();
@@ -2978,7 +2978,7 @@ void FKismetCompilerContext::ProcessOneFunctionGraph(UEdGraph* SourceGraph)
 	// When compiling only the skeleton class, we want the UFunction to be generated and processed so it contains all the local variables, this is unsafe to do during any other compilation mode
 	if (ValidateGraphIsWellFormed(FunctionGraph) || CompileOptions.CompileType == EKismetCompileType::SkeletonOnly)
 	{
-		FKismetFunctionContext& Context = *new (FunctionList) FKismetFunctionContext(MessageLog, Schema, NewClass, Blueprint);
+		FKismetFunctionContext& Context = *new (FunctionList)FKismetFunctionContext(MessageLog, Schema, NewClass, Blueprint, CompileOptions.DoesRequireCppCodeGeneration());
 		Context.SourceGraph = FunctionGraph;
 
 		if(FBlueprintEditorUtils::IsDelegateSignatureGraph(SourceGraph))
@@ -3074,7 +3074,7 @@ void FKismetCompilerContext::CreateFunctionList()
 
 FKismetFunctionContext* FKismetCompilerContext::CreateFunctionContext()
 {
-	return new (FunctionList)FKismetFunctionContext(MessageLog, Schema, NewClass, Blueprint);
+	return new (FunctionList)FKismetFunctionContext(MessageLog, Schema, NewClass, Blueprint, CompileOptions.DoesRequireCppCodeGeneration());
 }
 
 /** Compile a blueprint into a class and a set of functions */
@@ -3477,15 +3477,28 @@ void FKismetCompilerContext::Compile()
 		}
 
 		// Generate code thru the backend(s)
-		if (bDisplayCpp && bIsFullCompile)
+		if ((bDisplayCpp && bIsFullCompile) || CompileOptions.DoesRequireCppCodeGeneration())
 		{
-			FKismetCppBackend Backend_CPP(Schema, *this);
+			TUniquePtr<IKismetCppBackend> Backend_CPP(IKismetCppBackend::Create(Schema, *this));
 
 			// The C++ backend is currently only for debugging, so it's only run if the output will be visible
-			Backend_CPP.GenerateCodeFromClass(NewClass, FunctionList, !bIsFullCompile);
+			Backend_CPP->GenerateCodeFromClass(NewClass, FunctionList, !bIsFullCompile);
 		
-			UE_LOG(LogK2Compiler, Log, TEXT("[header]\n\n\n%s"), *Backend_CPP.Header);
-			UE_LOG(LogK2Compiler, Log, TEXT("[body]\n\n\n%s"), *Backend_CPP.Body);
+			if (CompileOptions.OutHeaderSourceCode.IsValid())
+			{
+				*CompileOptions.OutHeaderSourceCode = Backend_CPP->GetHeader();
+			}
+
+			if (CompileOptions.OutCppSourceCode.IsValid())
+			{
+				*CompileOptions.OutCppSourceCode = Backend_CPP->GetBody();
+			}
+
+			if (bDisplayCpp)
+			{
+				UE_LOG(LogK2Compiler, Log, TEXT("[header]\n\n\n%s"), *Backend_CPP->GetHeader());
+				UE_LOG(LogK2Compiler, Log, TEXT("[body]\n\n\n%s"), *Backend_CPP->GetBody());
+			}
 		}
 
 		// Always run the VM backend, it's needed for more than just debug printing
@@ -3666,7 +3679,7 @@ void FKismetCompilerContext::SetCanEverTickForActor()
 
 bool FKismetCompilerContext::UsePersistentUberGraphFrame() const
 {
-	return UBlueprintGeneratedClass::UsePersistentUberGraphFrame();
+	return UBlueprintGeneratedClass::UsePersistentUberGraphFrame() && !CompileOptions.DoesRequireCppCodeGeneration();
 }
 
 #undef LOCTEXT_NAMESPACE
diff --git a/Engine/Source/Editor/KismetCompiler/Private/KismetCompilerMisc.cpp b/Engine/Source/Editor/KismetCompiler/Private/KismetCompilerMisc.cpp
--- a/Engine/Source/Editor/KismetCompiler/Private/KismetCompilerMisc.cpp
+++ b/Engine/Source/Editor/KismetCompiler/Private/KismetCompilerMisc.cpp
@@ -977,7 +977,7 @@ KISMETCOMPILER_API FString FNetNameMapping::MakeBaseName<UEdGraphNode>(const UEd
 //////////////////////////////////////////////////////////////////////////
 // FKismetFunctionContext
 
-FKismetFunctionContext::FKismetFunctionContext(FCompilerResultsLog& InMessageLog, UEdGraphSchema_K2* InSchema, UBlueprintGeneratedClass* InNewClass, UBlueprint* InBlueprint)
+FKismetFunctionContext::FKismetFunctionContext(FCompilerResultsLog& InMessageLog, UEdGraphSchema_K2* InSchema, UBlueprintGeneratedClass* InNewClass, UBlueprint* InBlueprint, bool bInGeneratingCpp)
 	: Blueprint(InBlueprint)
 	, SourceGraph(NULL)
 	, EntryPoint(NULL)
@@ -994,6 +994,7 @@ FKismetFunctionContext::FKismetFunctionContext(FCompilerResultsLog& InMessageLog
 	, bIsSimpleStubGraphWithNoParams(false)
 	, NetFlags(0)
 	, SourceEventFromStubGraph(NULL)
+	, bGeneratingCpp(bInGeneratingCpp)
 {
 	NetNameMap = new FNetNameMapping();
 	bAllocatedNetNameMap = true;
@@ -1190,8 +1191,9 @@ FBPTerminal* FKismetFunctionContext::CreateLocalTerminalFromPinAutoChooseScope(U
 
 	const bool bRegularBlueprint = Blueprint && (Blueprint->GetClass() == UBlueprint::StaticClass());
 	static FBoolConfigValueHelper UseLocalGraphVariables(TEXT("Kismet"), TEXT("bUseLocalGraphVariables"), GEngineIni);
+	const bool bUseLocalGraphVariables = UseLocalGraphVariables || bGeneratingCpp;
 	const bool OutputPin = EEdGraphPinDirection::EGPD_Output == Net->Direction;
-	if (bSharedTerm && bRegularBlueprint && UseLocalGraphVariables && OutputPin)
+	if (bSharedTerm && bRegularBlueprint && bUseLocalGraphVariables && OutputPin)
 	{
 		BP_SCOPED_COMPILER_EVENT_STAT(EKismetCompilerStats_ChooseTerminalScope);
 
diff --git a/Engine/Source/Editor/KismetCompiler/Private/KismetCompilerModule.cpp b/Engine/Source/Editor/KismetCompiler/Private/KismetCompilerModule.cpp
--- a/Engine/Source/Editor/KismetCompiler/Private/KismetCompilerModule.cpp
+++ b/Engine/Source/Editor/KismetCompiler/Private/KismetCompilerModule.cpp
@@ -108,7 +108,7 @@ void FKismet2CompilerModule::CompileBlueprint(class UBlueprint* Blueprint, const
 
 	const bool bIsBrandNewBP = (Blueprint->SkeletonGeneratedClass == NULL) && (Blueprint->GeneratedClass == NULL) && (Blueprint->ParentClass != NULL) && !CompileOptions.bIsDuplicationInstigated;
 
-	if( CompileOptions.CompileType != EKismetCompileType::BytecodeOnly )
+	if ((CompileOptions.CompileType != EKismetCompileType::BytecodeOnly) && (CompileOptions.CompileType != EKismetCompileType::Cpp))
 	{
 		BP_SCOPED_COMPILER_EVENT_NAME(TEXT("Compile Skeleton Class"));
 
diff --git a/Engine/Source/Editor/KismetCompiler/Public/KismetCompiledFunctionContext.h b/Engine/Source/Editor/KismetCompiler/Public/KismetCompiledFunctionContext.h
--- a/Engine/Source/Editor/KismetCompiler/Public/KismetCompiledFunctionContext.h
+++ b/Engine/Source/Editor/KismetCompiler/Public/KismetCompiledFunctionContext.h
@@ -91,8 +91,10 @@ public:
 
 	// Stored calls of latent function (on current class), needed to tell if blueprint should be tickable
 	TArray< UK2Node_CallFunction* > LatentFunctionCalls;
+
+	bool bGeneratingCpp;
 public:
-	FKismetFunctionContext(FCompilerResultsLog& InMessageLog, UEdGraphSchema_K2* InSchema, UBlueprintGeneratedClass* InNewClass, UBlueprint* InBlueprint);
+	FKismetFunctionContext(FCompilerResultsLog& InMessageLog, UEdGraphSchema_K2* InSchema, UBlueprintGeneratedClass* InNewClass, UBlueprint* InBlueprint, bool bInGeneratingCpp);
 	
 	~FKismetFunctionContext();
 
diff --git a/Engine/Source/Editor/UnrealEd/Private/Kismet2/Kismet2.cpp b/Engine/Source/Editor/UnrealEd/Private/Kismet2/Kismet2.cpp
--- a/Engine/Source/Editor/UnrealEd/Private/Kismet2/Kismet2.cpp
+++ b/Engine/Source/Editor/UnrealEd/Private/Kismet2/Kismet2.cpp
@@ -553,6 +553,36 @@ void FKismetEditorUtilities::RecompileBlueprintBytecode(UBlueprint* BlueprintObj
 	}
 }
 
+/** Recompiles the bytecode of a blueprint only.  Should only be run for recompiling dependencies during compile on load */
+void FKismetEditorUtilities::GenerateCppCode(UBlueprint* InBlueprintObj, TSharedPtr<FString> OutHeaderSource, TSharedPtr<FString> OutCppSource)
+{
+	check(InBlueprintObj);
+	check(InBlueprintObj->GetOutermost() != GetTransientPackage());
+	checkf(InBlueprintObj->GeneratedClass, TEXT("Invalid generated class for %s"), *InBlueprintObj->GetName());
+	check(OutHeaderSource.IsValid());
+	check(OutCppSource.IsValid());
+
+	//TGuardValue<bool> DuplicatingReadOnly(InBlueprintObj->bDuplicatingReadOnly, true);
+	{
+		auto BlueprintObj = DuplicateObject<UBlueprint>(InBlueprintObj, GetTransientPackage(), *InBlueprintObj->GetName());
+		{
+			IKismetCompilerInterface& Compiler = FModuleManager::LoadModuleChecked<IKismetCompilerInterface>(KISMET_COMPILER_MODULENAME);
+
+			TGuardValue<bool> GuardTemplateNameFlag(GCompilingBlueprint, true);
+			FCompilerResultsLog Results;
+
+			FKismetCompilerOptions CompileOptions;
+			CompileOptions.CompileType = EKismetCompileType::Cpp;
+			CompileOptions.OutCppSourceCode = OutCppSource;
+			CompileOptions.OutHeaderSourceCode = OutHeaderSource;
+			Compiler.CompileBlueprint(BlueprintObj, CompileOptions, Results);
+		}
+		BlueprintObj->RemoveGeneratedClasses();
+		BlueprintObj->ClearFlags(RF_Standalone);
+		BlueprintObj->MarkPendingKill();
+	}
+}
+
 /** Tries to make sure that a blueprint is conformed to its native parent, in case any native class flags have changed */
 void FKismetEditorUtilities::ConformBlueprintFlagsAndComponents(UBlueprint* BlueprintObj)
 {
diff --git a/Engine/Source/Editor/UnrealEd/Public/Kismet2/KismetEditorUtilities.h b/Engine/Source/Editor/UnrealEd/Public/Kismet2/KismetEditorUtilities.h
--- a/Engine/Source/Editor/UnrealEd/Public/Kismet2/KismetEditorUtilities.h
+++ b/Engine/Source/Editor/UnrealEd/Public/Kismet2/KismetEditorUtilities.h
@@ -40,6 +40,8 @@ public:
 	/** Recompiles the bytecode of a blueprint only.  Should only be run for recompiling dependencies during compile on load */
 	static void RecompileBlueprintBytecode(UBlueprint* BlueprintObj, TArray<UObject*>* ObjLoaded = NULL);
 
+	static void GenerateCppCode(UBlueprint* BlueprintObj, TSharedPtr<FString> OutHeaderSource, TSharedPtr<FString> OutCppSource);
+
 	/** Tries to make sure that a data-only blueprint is conformed to its native parent, in case any native class flags have changed */
 	static void ConformBlueprintFlagsAndComponents(UBlueprint* BlueprintObj);
 
diff --git a/Engine/Source/Runtime/Engine/Classes/Engine/Blueprint.h b/Engine/Source/Runtime/Engine/Classes/Engine/Blueprint.h
--- a/Engine/Source/Runtime/Engine/Classes/Engine/Blueprint.h
+++ b/Engine/Source/Runtime/Engine/Classes/Engine/Blueprint.h
@@ -55,6 +55,7 @@ namespace EKismetCompileType
 		Full,
 		StubAfterFailure, 
 		BytecodeOnly,
+		Cpp,
 	};
 };
 
@@ -70,9 +71,19 @@ public:
 	/** Whether or not this compile is for a duplicated blueprint */
 	bool bIsDuplicationInstigated;
 
+	TSharedPtr<FString> OutHeaderSourceCode;
+	TSharedPtr<FString> OutCppSourceCode;
+
+	bool DoesRequireCppCodeGeneration() const
+	{
+		return (CompileType == EKismetCompileType::Cpp);
+	}
+
 	bool DoesRequireBytecodeGeneration() const
 	{
-		return (CompileType == EKismetCompileType::Full) || (CompileType == EKismetCompileType::BytecodeOnly);
+		return (CompileType == EKismetCompileType::Full) 
+			|| (CompileType == EKismetCompileType::BytecodeOnly) 
+			|| (CompileType == EKismetCompileType::Cpp);
 	}
 
 	/** Whether or not this compile type should operate on the generated class of the blueprint, as opposed to just the skeleton */
