CIS fix: virtual desctructor in IKismetCppBackend

diff --git a/Engine/Source/Editor/KismetCompiler/Private/KismetCompilerBackend.h b/Engine/Source/Editor/KismetCompiler/Private/KismetCompilerBackend.h
--- a/Engine/Source/Editor/KismetCompiler/Private/KismetCompilerBackend.h
+++ b/Engine/Source/Editor/KismetCompiler/Private/KismetCompilerBackend.h
@@ -26,6 +26,8 @@ public:
 	virtual const FString& GetBody() const = 0;
 	virtual const FString& GetHeader() const = 0;
 	virtual void GenerateCodeFromClass(UClass* SourceClass, TIndirectArray<FKismetFunctionContext>& Functions, bool bGenerateStubsOnly = false) = 0;
+
+	virtual ~IKismetCppBackend() {}
 };
 
 //////////////////////////////////////////////////////////////////////////
