UDS Fix: Crash when modifying data layout of instanced UDS

diff --git a/Engine/Source/Editor/KismetCompiler/Private/UserDefinedStructureCompilerUtils.cpp b/Engine/Source/Editor/KismetCompiler/Private/UserDefinedStructureCompilerUtils.cpp
--- a/Engine/Source/Editor/KismetCompiler/Private/UserDefinedStructureCompilerUtils.cpp
+++ b/Engine/Source/Editor/KismetCompiler/Private/UserDefinedStructureCompilerUtils.cpp
@@ -23,11 +23,11 @@ struct FUserDefinedStructureCompilerInner
 				const FString ReinstancedName = FString::Printf(TEXT("STRUCT_REINST_%s"), *StructureToReinstance->GetName());
 				const FName UniqueName = MakeUniqueObjectName(GetTransientPackage(), UUserDefinedStruct::StaticClass(), FName(*ReinstancedName));
 
-				const bool OldIsDuplicatingClassForReinstancing = GIsDuplicatingClassForReinstancing;
-				GIsDuplicatingClassForReinstancing = true;
+				TGuardValue<bool> IsDuplicatingClassForReinstancing(GIsDuplicatingClassForReinstancing, true);
 				DuplicatedStruct = (UUserDefinedStruct*)StaticDuplicateObject(StructureToReinstance, GetTransientPackage(), *UniqueName.ToString(), ~RF_Transactional); 
-				GIsDuplicatingClassForReinstancing = OldIsDuplicatingClassForReinstancing;
 			}
+			DuplicatedStruct->Bind();
+			DuplicatedStruct->StaticLink(true);
 			DuplicatedStruct->PrimaryStruct = StructureToReinstance;
 			DuplicatedStruct->Status = EUserDefinedStructureStatus::UDSS_Duplicate;
 			DuplicatedStruct->SetFlags(RF_Transient);
