[UE-2559] Making a couple BP compiler errors more understandable.

diff --git a/Engine/Source/Editor/BlueprintGraph/Private/EdGraphSchema_K2.cpp b/Engine/Source/Editor/BlueprintGraph/Private/EdGraphSchema_K2.cpp
--- a/Engine/Source/Editor/BlueprintGraph/Private/EdGraphSchema_K2.cpp
+++ b/Engine/Source/Editor/BlueprintGraph/Private/EdGraphSchema_K2.cpp
@@ -2164,17 +2164,22 @@ FString UEdGraphSchema_K2::IsPinDefaultValid(const UEdGraphPin* Pin, const FStri
 	const bool bIsReference = Pin->PinType.bIsReference;
 	const bool bIsAutoCreateRefTerm = IsAutoCreateRefTerm(Pin);
 
+	FFormatNamedArguments MessageArgs;
+	MessageArgs.Add(TEXT("PinName"), Pin->GetDisplayName());
+
 	if (OwningBP->BlueprintType != BPTYPE_Interface)
 	{
 		if( !bIsAutoCreateRefTerm )
 		{
 			if( bIsArray )
 			{
-				return TEXT("Literal values are not allowed for array parameters.  Use a Make Array node instead");
+				FText MsgFormat = LOCTEXT("BadArrayDefaultVal", "Array inputs (like '{PinName}') must have an input wired into them (try connecting a MakeArray node).");
+				return FText::Format(MsgFormat, MessageArgs).ToString();
 			}
 			else if( bIsReference )
 			{
-				return TEXT("Literal values are not allowed for pass-by-reference parameters.");
+				FText MsgFormat = LOCTEXT("BadRefDefaultVal", "'{PinName}' must have an input wired into it (\"by ref\" params expect a valid input to operate on).");
+				return FText::Format(MsgFormat, MessageArgs).ToString();
 			}
 		}
 	}
diff --git a/Engine/Source/Editor/KismetCompiler/Private/KismetCompilerMisc.cpp b/Engine/Source/Editor/KismetCompiler/Private/KismetCompilerMisc.cpp
--- a/Engine/Source/Editor/KismetCompiler/Private/KismetCompilerMisc.cpp
+++ b/Engine/Source/Editor/KismetCompiler/Private/KismetCompilerMisc.cpp
@@ -933,9 +933,19 @@ void FNodeHandlingFunctor::RegisterNets(FKismetFunctionContext& Context, UEdGrap
 				{
 					// Make sure the default value is valid
 					FString DefaultAllowedResult = CompilerContext.GetSchema()->IsCurrentPinDefaultValid(Net);
-					if (DefaultAllowedResult != TEXT(""))
+					if (!DefaultAllowedResult.IsEmpty())
 					{
-						CompilerContext.MessageLog.Error(*FString::Printf(*LOCTEXT("InvalidDefaultValue_Error", "Default value '%s' for @@ is invalid: '%s'").ToString(), *(Net->GetDefaultAsString()), *DefaultAllowedResult), Net);
+						FText ErrorFormat = LOCTEXT("InvalidDefault_Error", "The current value of the '@@' pin is invalid: {0}");
+						const FText InvalidReasonText = FText::FromString(DefaultAllowedResult);
+
+						FText DefaultValue = FText::FromString(Net->GetDefaultAsString());
+						if (!DefaultValue.IsEmpty())
+						{
+							ErrorFormat = LOCTEXT("InvalidDefaultVal_Error", "The current value ({1}) of the '@@' pin is invalid: {0}");
+						}
+
+						FString ErrorString = FText::Format(ErrorFormat, InvalidReasonText, DefaultValue).ToString();
+						CompilerContext.MessageLog.Error(*ErrorString, Net);
 
 						// Skip over these properties if they are array or ref properties, because the backend can't emit valid code for them
 						if( Pin->PinType.bIsArray || Pin->PinType.bIsReference )
