more verbose output on unsupported VM data types

diff --git a/Engine/Source/Editor/KismetCompiler/Private/KismetCompilerVMBackend.cpp b/Engine/Source/Editor/KismetCompiler/Private/KismetCompilerVMBackend.cpp
--- a/Engine/Source/Editor/KismetCompiler/Private/KismetCompilerVMBackend.cpp
+++ b/Engine/Source/Editor/KismetCompiler/Private/KismetCompilerVMBackend.cpp
@@ -539,14 +539,14 @@ public:
 				}
 				else 
 				{
-					ensureMsg(false, TEXT("It is not possible to express this interface property as a literal value!"));
+					ensureMsgf(false, TEXT("It is not possible to express this interface property as a literal value! (%s)"), *CoerceProperty->GetFullName());
 				}
 			}
 			// else if (CoerceProperty->IsA(UMulticastDelegateProperty::StaticClass()))
 			// Cannot assign a literal to a multicast delegate; it should be added instead of assigned
 			else
 			{
-				ensureMsg(false, TEXT("It is not possible to express this type as a literal value!"));
+				ensureMsgf(false, TEXT("It is not possible to express this type as a literal value! (%s)"), *CoerceProperty->GetFullName());
 			}
 		}
 		else
