Make implicit blueprint created component properties non-transactional

diff --git a/Engine/Source/Editor/KismetCompiler/Private/KismetCompiler.cpp b/Engine/Source/Editor/KismetCompiler/Private/KismetCompiler.cpp
--- a/Engine/Source/Editor/KismetCompiler/Private/KismetCompiler.cpp
+++ b/Engine/Source/Editor/KismetCompiler/Private/KismetCompiler.cpp
@@ -572,7 +572,7 @@ void FKismetCompilerContext::CreateClassVariablesFromBlueprint()
 					const FString CategoryName = Node->CategoryName != NAME_None ? Node->CategoryName.ToString() : Blueprint->GetName();
 					
 					NewProperty->SetMetaData(TEXT("Category"), *CategoryName);
-					NewProperty->SetPropertyFlags(CPF_BlueprintVisible);
+					NewProperty->SetPropertyFlags(CPF_BlueprintVisible|CPF_NonTransactional);
 				}
 			}
 		}
