[UE-5051] Literal self terms now emit the proper expression for our custom-thunk array nodes.

diff --git a/Engine/Source/Editor/KismetCompiler/Private/KismetCompilerVMBackend.cpp b/Engine/Source/Editor/KismetCompiler/Private/KismetCompilerVMBackend.cpp
--- a/Engine/Source/Editor/KismetCompiler/Private/KismetCompilerVMBackend.cpp
+++ b/Engine/Source/Editor/KismetCompiler/Private/KismetCompilerVMBackend.cpp
@@ -386,10 +386,30 @@ public:
 			}
 			else if (CoerceProperty->IsA(UIntProperty::StaticClass()))
 			{
-				//@TODO: There are smaller encodings EX_IntZero, EX_IntOne, EX_IntConstByte available which could be used instead when the value fits
-				int32 Value = FCString::Atoi(*(Term->Name));
-				Writer << EX_IntConst;
-				Writer << Value;
+				// In certain cases (like UKismetArrayLibrary functions), we have
+				// polymorphic functions that provide their own "custom thunk" 
+				// (custom execution code). The actual function acts as a 
+				// template, where the parameter types can be changed out for 
+				// other types (much like c++ template functions, the "custom 
+				// thunk" is generic). Traditionally, we use integer refs as the 
+				// place holder type (that's why this is nested in a UIntProperty 
+				// check)... Complications arise here, when we try to emit 
+				// literal values fed into the function when they don't match 
+				// the template's (int) type. For most types, this here is 
+				// circumvented with AutoCreateRefTerm, but when it is a self 
+				// (literal) node we still end up here. So, we try to detect and 
+				// handle that case here.
+				if ((Term->Type.PinSubCategory == Schema->PN_Self) && CoerceProperty->HasAnyPropertyFlags(CPF_ReferenceParm))
+				{
+					Writer << EX_Self;
+				}
+				else
+				{
+					//@TODO: There are smaller encodings EX_IntZero, EX_IntOne, EX_IntConstByte available which could be used instead when the value fits
+					int32 Value = FCString::Atoi(*(Term->Name));
+					Writer << EX_IntConst;
+					Writer << Value;
+				}
 			}
 			else if (CoerceProperty->IsA(UByteProperty::StaticClass()))
 			{
