BP/UMG - Allowing function graph compiles to specify if the function should be compiled as internal use only.  Needed in certain cases in UMG when compiling functions you don't want users to see in the reflection info.  Similar to the need when compiling events.

diff --git a/Engine/Source/Editor/KismetCompiler/Private/AnimBlueprintCompiler.cpp b/Engine/Source/Editor/KismetCompiler/Private/AnimBlueprintCompiler.cpp
--- a/Engine/Source/Editor/KismetCompiler/Private/AnimBlueprintCompiler.cpp
+++ b/Engine/Source/Editor/KismetCompiler/Private/AnimBlueprintCompiler.cpp
@@ -1074,7 +1074,7 @@ void FAnimBlueprintCompiler::MergeUbergraphPagesIn(UEdGraph* Ubergraph)
 	}
 }
 
-void FAnimBlueprintCompiler::ProcessOneFunctionGraph(UEdGraph* SourceGraph)
+void FAnimBlueprintCompiler::ProcessOneFunctionGraph(UEdGraph* SourceGraph, bool bInternalFunction)
 {
 	if (SourceGraph->Schema->IsChildOf(UAnimationGraphSchema::StaticClass()))
 	{
@@ -1092,7 +1092,7 @@ void FAnimBlueprintCompiler::ProcessOneFunctionGraph(UEdGraph* SourceGraph)
 	else
 	{
 		// Let the regular K2 compiler handle this one
-		Super::ProcessOneFunctionGraph(SourceGraph);
+		Super::ProcessOneFunctionGraph(SourceGraph, bInternalFunction);
 	}
 }
 
diff --git a/Engine/Source/Editor/KismetCompiler/Private/AnimBlueprintCompiler.h b/Engine/Source/Editor/KismetCompiler/Private/AnimBlueprintCompiler.h
--- a/Engine/Source/Editor/KismetCompiler/Private/AnimBlueprintCompiler.h
+++ b/Engine/Source/Editor/KismetCompiler/Private/AnimBlueprintCompiler.h
@@ -28,7 +28,7 @@ protected:
 	// Implementation of FKismetCompilerContext interface
 	virtual UEdGraphSchema_K2* CreateSchema() override;
 	virtual void MergeUbergraphPagesIn(UEdGraph* Ubergraph) override;
-	virtual void ProcessOneFunctionGraph(UEdGraph* SourceGraph) override;
+	virtual void ProcessOneFunctionGraph(UEdGraph* SourceGraph, bool bInternalFunction = false) override;
 	virtual void CreateFunctionList() override;
 	virtual void SpawnNewClass(const FString& NewClassName) override;
 	virtual void CopyTermDefaultsToDefaultObject(UObject* DefaultObject) override;
diff --git a/Engine/Source/Editor/KismetCompiler/Private/KismetCompiler.cpp b/Engine/Source/Editor/KismetCompiler/Private/KismetCompiler.cpp
--- a/Engine/Source/Editor/KismetCompiler/Private/KismetCompiler.cpp
+++ b/Engine/Source/Editor/KismetCompiler/Private/KismetCompiler.cpp
@@ -2960,7 +2960,7 @@ void FKismetCompilerContext::ResetErrorFlags(UEdGraph* Graph) const
 /**
  * Merges macros/subgraphs into the graph and validates it, creating a function list entry if it's reasonable.
  */
-void FKismetCompilerContext::ProcessOneFunctionGraph(UEdGraph* SourceGraph)
+void FKismetCompilerContext::ProcessOneFunctionGraph(UEdGraph* SourceGraph, bool bInternalFunction)
 {
 	BP_SCOPED_COMPILER_EVENT_STAT(EKismetCompilerStats_ProcessFunctionGraph);
 
@@ -2996,6 +2996,11 @@ void FKismetCompilerContext::ProcessOneFunctionGraph(UEdGraph* SourceGraph)
 		{
 			Context.MarkAsConstFunction();
 		}
+
+		if ( bInternalFunction )
+		{
+			Context.MarkAsInternalOrCppUseOnly();
+		}
 	}
 }
 
diff --git a/Engine/Source/Editor/KismetCompiler/Public/KismetCompiler.h b/Engine/Source/Editor/KismetCompiler/Public/KismetCompiler.h
--- a/Engine/Source/Editor/KismetCompiler/Public/KismetCompiler.h
+++ b/Engine/Source/Editor/KismetCompiler/Public/KismetCompiler.h
@@ -271,7 +271,7 @@ protected:
 	/**
 	 * Merges macros/subgraphs into the graph and validates it, creating a function list entry if it's reasonable.
 	 */
-	virtual void ProcessOneFunctionGraph(UEdGraph* SourceGraph);
+	virtual void ProcessOneFunctionGraph(UEdGraph* SourceGraph, bool bInternalFunction = false);
 
 	/**
 	 * Picks the name to use for an autogenerated event stub
