BP: SkeletonOnly compilation doesn't call Transform and Validate on nodes.

UE-3686 BP: Erroneous errors when adding a timeline node

diff --git a/Engine/Source/Editor/KismetCompiler/Private/KismetCompiler.cpp b/Engine/Source/Editor/KismetCompiler/Private/KismetCompiler.cpp
--- a/Engine/Source/Editor/KismetCompiler/Private/KismetCompiler.cpp
+++ b/Engine/Source/Editor/KismetCompiler/Private/KismetCompiler.cpp
@@ -1112,11 +1112,14 @@ void FKismetCompilerContext::PrecompileFunction(FKismetFunctionContext& Context)
 		// Find the connected subgraph starting at the root node and prune out unused nodes
 		PruneIsolatedNodes(Context.RootSet, Context.SourceGraph->Nodes);
 
-		// Check if self pins are connected after PruneIsolatedNodes, to avoid errors from isolated nodes.
-		ValidateSelfPinsInGraph(Context.SourceGraph);
+		if (bIsFullCompile)
+		{
+			// Check if self pins are connected after PruneIsolatedNodes, to avoid errors from isolated nodes.
+			ValidateSelfPinsInGraph(Context.SourceGraph);
 
-		// Transforms
-		TransformNodes(Context);
+			// Transforms
+			TransformNodes(Context);
+		}
 
 		//Now we can safely remove automatically added WorldContext pin from static function.
 		Context.EntryPoint->RemoveUnnecessaryAutoWorldContext();
@@ -2661,15 +2664,16 @@ void FKismetCompilerContext::CreateAndProcessUbergraph()
 			{
 				const UEdGraphNode* Node = ConsolidatedEventGraph->Nodes[ChildIndex];
 				const int32 SavedErrorCount = MessageLog.NumErrors;
-				ValidateNode(Node);
+				UK2Node_Event* SrcEventNode = Cast<UK2Node_Event>(ConsolidatedEventGraph->Nodes[ChildIndex]);
+				if (bIsFullCompile || SrcEventNode)
+				{
+					ValidateNode(Node);
+				}
 
 				// If the node didn't generate any errors then generate function stubs for event entry nodes etc.
-				if (SavedErrorCount == MessageLog.NumErrors)
+				if ((SavedErrorCount == MessageLog.NumErrors) && SrcEventNode)
 				{
-					if (UK2Node_Event* SrcEventNode = Cast<UK2Node_Event>(ConsolidatedEventGraph->Nodes[ChildIndex]))
-					{
-						CreateFunctionStubForEvent(SrcEventNode, Blueprint);
-					}
+					CreateFunctionStubForEvent(SrcEventNode, Blueprint);
 				}
 			}
 		}
