Added an ini setting (bSkeletonInheritSkeletonClasse) to allow Skeleton Classes to inherit parent BP's Skeleton Class instead of the fully compiled class.

Defaults to false for testing.

When enabled, variable type changes that will impact child Blueprints will no longer force a (potentially lengthy) full compile.

diff --git a/Engine/Config/BaseEngine.ini b/Engine/Config/BaseEngine.ini
--- a/Engine/Config/BaseEngine.ini
+++ b/Engine/Config/BaseEngine.ini
@@ -1463,6 +1463,7 @@ bReinstanceOnlyWhenNecessary=true
 bChangeDefaultValueWithoutReinstancing=true
 bDisplaysLayout=false
 bNativeCodeGenerationTool=false
+bSkeletonInheritSkeletonClasses=false
 
 [/Script/Engine.Blueprint]
 bRecompileOnLoad=true
diff --git a/Engine/Source/Editor/BlueprintGraph/Private/BlueprintVariableNodeSpawner.cpp b/Engine/Source/Editor/BlueprintGraph/Private/BlueprintVariableNodeSpawner.cpp
--- a/Engine/Source/Editor/BlueprintGraph/Private/BlueprintVariableNodeSpawner.cpp
+++ b/Engine/Source/Editor/BlueprintGraph/Private/BlueprintVariableNodeSpawner.cpp
@@ -74,10 +74,6 @@ UBlueprintVariableNodeSpawner* UBlueprintVariableNodeSpawner::Create(TSubclassOf
 			UClass* OwnerClass = Property->GetOwnerClass();
 
 			// We need to use a generated class instead of a skeleton class for IsChildOf, so if the OwnerClass has a Blueprint, grab the GeneratedClass
-			if(OwnerClass)
-			{
-				OwnerClass = OwnerClass->GetAuthoritativeClass();
-			}
 			bool const bIsSelfContext = Blueprint->SkeletonGeneratedClass->GetAuthoritativeClass() == OwnerClass || Blueprint->SkeletonGeneratedClass->IsChildOf(OwnerClass);
 
 			UK2Node_Variable* VarNode = CastChecked<UK2Node_Variable>(NewNode);
diff --git a/Engine/Source/Editor/KismetCompiler/Private/KismetCompiler.cpp b/Engine/Source/Editor/KismetCompiler/Private/KismetCompiler.cpp
--- a/Engine/Source/Editor/KismetCompiler/Private/KismetCompiler.cpp
+++ b/Engine/Source/Editor/KismetCompiler/Private/KismetCompiler.cpp
@@ -10,6 +10,7 @@
 #include "Editor/UnrealEd/Public/Kismet2/KismetDebugUtilities.h"
 #include "Editor/UnrealEd/Public/Kismet2/KismetReinstanceUtilities.h"
 #include "Editor/UnrealEd/Public/Kismet2/BlueprintEditorUtils.h"
+#include "Editor/UnrealEd/Public/Kismet2/KismetEditorUtilities.h"
 #include "Editor/UnrealEd/Public/ScriptDisassembler.h"
 #include "K2Node_PlayMovieScene.h"
 #include "RuntimeMovieScenePlayer.h"
@@ -153,6 +154,15 @@ void FKismetCompilerContext::CleanAndSanitizeClass(UBlueprintGeneratedClass* Cla
 	UClass* TransientClass = ConstructObject<UBlueprintGeneratedClass>(UBlueprintGeneratedClass::StaticClass(), GetTransientPackage(), TransientClassName, RF_Public|RF_Transient);
 	
 	UClass* ParentClass = Blueprint->ParentClass;
+
+	if(UBlueprintGeneratedClass::CompileSkeletonClassesInheritSkeletonClasses() && FKismetEditorUtilities::IsClassABlueprintSkeleton(ClassToClean))
+	{
+		if(UBlueprint* BlueprintParent = Cast<UBlueprint>(Blueprint->ParentClass->ClassGeneratedBy))
+		{
+			ParentClass = BlueprintParent->SkeletonGeneratedClass;
+		}
+	}
+
 	if( ParentClass == NULL )
 	{
 		ParentClass = UObject::StaticClass();
@@ -3260,9 +3270,18 @@ void FKismetCompilerContext::Compile()
 
 	NewClass->ClassGeneratedBy = Blueprint;
 
-	NewClass->SetSuperStruct(Blueprint->ParentClass);
-	NewClass->ClassFlags |= (Blueprint->ParentClass->ClassFlags & CLASS_Inherit);
-	NewClass->ClassCastFlags |= Blueprint->ParentClass->ClassCastFlags;
+	UClass* ParentClass = nullptr;
+	if(UBlueprintGeneratedClass::CompileSkeletonClassesInheritSkeletonClasses())
+	{
+		ParentClass = NewClass->ClassWithin;
+	}
+	else
+	{
+		ParentClass = Blueprint->ParentClass;
+	}
+	NewClass->SetSuperStruct(ParentClass);
+	NewClass->ClassFlags |= (ParentClass->ClassFlags & CLASS_Inherit);
+	NewClass->ClassCastFlags |= ParentClass->ClassCastFlags;
 	
 	if(Blueprint->bGenerateConstClass)
 	{
diff --git a/Engine/Source/Editor/UnrealEd/Private/Kismet2/BlueprintEditorUtils.cpp b/Engine/Source/Editor/UnrealEd/Private/Kismet2/BlueprintEditorUtils.cpp
--- a/Engine/Source/Editor/UnrealEd/Private/Kismet2/BlueprintEditorUtils.cpp
+++ b/Engine/Source/Editor/UnrealEd/Private/Kismet2/BlueprintEditorUtils.cpp
@@ -3735,8 +3735,18 @@ void FBlueprintEditorUtils::ChangeMemberVariableType(UBlueprint* Blueprint, cons
 				{
 					Variable.VarType = NewPinType;
 
-					// Compile the Blueprint even if no loaded BPs are using the variable, this ensures that BPs refresh correctly when they are loaded
-					FKismetEditorUtilities::CompileBlueprint(Blueprint, false, true);
+					UClass* ParentClass = nullptr;
+					// When compiling so that skeleton classes inherit skeleton classes, marking the Blueprint structurally modified is enough
+					// when the children skeleton classes are rebuilt, they will receive the modified variables.
+					if(UBlueprintGeneratedClass::CompileSkeletonClassesInheritSkeletonClasses())
+					{
+						FBlueprintEditorUtils::MarkBlueprintAsStructurallyModified(Blueprint);
+					}
+					else
+					{
+						// Compile the Blueprint even if no loaded BPs are using the variable, this ensures that BPs refresh correctly when they are loaded
+						FKismetEditorUtilities::CompileBlueprint(Blueprint, false, true);
+					}
 
 					if(bBreakingVariableConnections)
 					{
diff --git a/Engine/Source/Editor/UnrealEd/Private/Kismet2/KismetReinstanceUtilities.cpp b/Engine/Source/Editor/UnrealEd/Private/Kismet2/KismetReinstanceUtilities.cpp
--- a/Engine/Source/Editor/UnrealEd/Private/Kismet2/KismetReinstanceUtilities.cpp
+++ b/Engine/Source/Editor/UnrealEd/Private/Kismet2/KismetReinstanceUtilities.cpp
@@ -30,6 +30,8 @@ FBlueprintCompileReinstancer::FBlueprintCompileReinstancer(UClass* InClassToRein
 {
 	if( InClassToReinstance != NULL )
 	{
+		bIsReinstancingSkeleton = FKismetEditorUtilities::IsClassABlueprintSkeleton(ClassToReinstance);
+
 		SaveClassFieldMapping(InClassToReinstance);
 
 		// Remember the initial CDO for the class being resinstanced
@@ -796,17 +798,18 @@ void FBlueprintCompileReinstancer::VerifyReplacement()
 void FBlueprintCompileReinstancer::ReparentChild(UBlueprint* ChildBP)
 {
 	check(ChildBP);
-	check(ChildBP->ParentClass == ClassToReinstance || ChildBP->ParentClass == DuplicatedClass);
 
 	UClass* SkeletonClass = ChildBP->SkeletonGeneratedClass;
 	UClass* GeneratedClass = ChildBP->GeneratedClass;
 
-	if( SkeletonClass )
+	const bool bReinstancingSkeleton = (UBlueprintGeneratedClass::CompileSkeletonClassesInheritSkeletonClasses() && bIsReinstancingSkeleton) || !UBlueprintGeneratedClass::CompileSkeletonClassesInheritSkeletonClasses();
+	if(  bReinstancingSkeleton && SkeletonClass )
 	{
 		ReparentChild(SkeletonClass);
 	}
 
-	if( GeneratedClass )
+	const bool bReinstancingGenerated = (UBlueprintGeneratedClass::CompileSkeletonClassesInheritSkeletonClasses() && !bIsReinstancingSkeleton) || !UBlueprintGeneratedClass::CompileSkeletonClassesInheritSkeletonClasses();
+	if( bReinstancingGenerated && GeneratedClass )
 	{
 		ReparentChild(GeneratedClass);
 	}
diff --git a/Engine/Source/Editor/UnrealEd/Public/Kismet2/KismetReinstanceUtilities.h b/Engine/Source/Editor/UnrealEd/Public/Kismet2/KismetReinstanceUtilities.h
--- a/Engine/Source/Editor/UnrealEd/Public/Kismet2/KismetReinstanceUtilities.h
+++ b/Engine/Source/Editor/UnrealEd/Public/Kismet2/KismetReinstanceUtilities.h
@@ -47,6 +47,9 @@ protected:
 	/** Don't call GC */
 	bool bSkipGarbageCollection;
 
+	/** Cached value, true if reinstancing a skeleton class or not */
+	bool bIsReinstancingSkeleton;
+
 	uint32 ClassToReinstanceDefaultValuesCRC;
 
 public:
diff --git a/Engine/Source/Runtime/Engine/Classes/Engine/BlueprintGeneratedClass.h b/Engine/Source/Runtime/Engine/Classes/Engine/BlueprintGeneratedClass.h
--- a/Engine/Source/Runtime/Engine/Classes/Engine/BlueprintGeneratedClass.h
+++ b/Engine/Source/Runtime/Engine/Classes/Engine/BlueprintGeneratedClass.h
@@ -479,6 +479,8 @@ public:
 	static FName GetUberGraphFrameName();
 	static bool UsePersistentUberGraphFrame();
 
+	static bool CompileSkeletonClassesInheritSkeletonClasses();
+
 #if WITH_EDITORONLY_DATA
 	FBlueprintDebugData DebugData;
 
diff --git a/Engine/Source/Runtime/Engine/Private/BlueprintGeneratedClass.cpp b/Engine/Source/Runtime/Engine/Private/BlueprintGeneratedClass.cpp
--- a/Engine/Source/Runtime/Engine/Private/BlueprintGeneratedClass.cpp
+++ b/Engine/Source/Runtime/Engine/Private/BlueprintGeneratedClass.cpp
@@ -584,4 +584,10 @@ bool UBlueprintGeneratedClass::UsePersistentUberGraphFrame()
 #else
 	return false;
 #endif
+}
+
+bool UBlueprintGeneratedClass::CompileSkeletonClassesInheritSkeletonClasses()
+{
+	static const FBoolConfigValueHelper SkeletonClassesInheritSkeletonClasses(TEXT("Kismet"), TEXT("bSkeletonInheritSkeletonClasses"), GEngineIni);
+	return SkeletonClassesInheritSkeletonClasses;
 }
\ No newline at end of file
