BP runtime optimisation:
Execution Flow Stack is used only when necessary. KCST_GotoReturn and KCST_GotoReturnIfNot added.
Number of jumps/gotos is reduced. Statements are sorted. Adjacent statements are merged.

diff --git a/Engine/Config/BaseEngine.ini b/Engine/Config/BaseEngine.ini
--- a/Engine/Config/BaseEngine.ini
+++ b/Engine/Config/BaseEngine.ini
@@ -1464,6 +1464,8 @@ bChangeDefaultValueWithoutReinstancing=true
 bDisplaysLayout=false
 bNativeCodeGenerationTool=false
 bSkeletonInheritSkeletonClasses=false
+bOptimizeExecutionFlowStack=true
+bOptimizeAdjacentStates=true
 
 [/Script/Engine.Blueprint]
 bRecompileOnLoad=true
diff --git a/Engine/Source/Editor/KismetCompiler/Private/KismetCompiler.cpp b/Engine/Source/Editor/KismetCompiler/Private/KismetCompiler.cpp
--- a/Engine/Source/Editor/KismetCompiler/Private/KismetCompiler.cpp
+++ b/Engine/Source/Editor/KismetCompiler/Private/KismetCompiler.cpp
@@ -49,6 +49,7 @@ DEFINE_STAT(EKismetCompilerStats_CleanAndSanitizeClass);
 DEFINE_STAT(EKismetCompilerStats_CreateClassVariables);
 DEFINE_STAT(EKismetCompilerStats_BindAndLinkClass);
 DEFINE_STAT(EKismetCompilerStats_ChecksumCDO);
+DEFINE_STAT(EKismetCompilerStats_ResolveCompiledStatements);
 		
 //////////////////////////////////////////////////////////////////////////
 // FKismetCompilerContext
@@ -1475,11 +1476,8 @@ void FKismetCompilerContext::PostcompileFunction(FKismetFunctionContext& Context
 {
 	BP_SCOPED_COMPILER_EVENT_STAT(EKismetCompilerStats_PostcompileFunction);
 
-	// Sort the 'linear execution list' again by likely execution order.
-	Context.FinalSortLinearExecList();
-
-	// Resolve goto links
-	Context.ResolveGotoFixups();
+	// The function links gotos, sorts statments, and merges adjacent ones. 
+	Context.ResolveStatements();
 
 	//@TODO: Code generation (should probably call backend here, not later)
 
diff --git a/Engine/Source/Editor/KismetCompiler/Private/KismetCompilerCppBackend.cpp b/Engine/Source/Editor/KismetCompiler/Private/KismetCompilerCppBackend.cpp
--- a/Engine/Source/Editor/KismetCompiler/Private/KismetCompilerCppBackend.cpp
+++ b/Engine/Source/Editor/KismetCompiler/Private/KismetCompilerCppBackend.cpp
@@ -1166,7 +1166,7 @@ void FKismetCppBackend::EmitGotoStatement(FKismetFunctionContext& FunctionContex
 		Emit(Body, *FString::Printf(TEXT("\t\t\tCurrentState = %s;\n"), *NextStateExpression));
 		Emit(Body, *FString::Printf(TEXT("\t\t\tbreak;\n")));
 	}
-	else if ((Statement.Type == KCST_GotoIfNot) || (Statement.Type == KCST_EndOfThreadIfNot))
+	else if ((Statement.Type == KCST_GotoIfNot) || (Statement.Type == KCST_EndOfThreadIfNot) || (Statement.Type == KCST_GotoReturnIfNot))
 	{
 		FString ConditionExpression;
 		ConditionExpression = TermToText(Statement.LHS, (UProperty*)(GetDefault<UBoolProperty>()));
@@ -1176,8 +1176,13 @@ void FKismetCppBackend::EmitGotoStatement(FKismetFunctionContext& FunctionContex
 
 		if (Statement.Type == KCST_EndOfThreadIfNot)
 		{
+			ensure(FunctionContext.bUseFlowStack);
 			Emit(Body, TEXT("\t\t\t\tCurrentState = (StateStack.Num() > 0) ? StateStack.Pop() : -1;\n"));
 		}
+		else if (Statement.Type == KCST_GotoReturnIfNot)
+		{
+			Emit(Body, TEXT("\t\t\t\tCurrentState = -1;\n"));
+		}
 		else
 		{
 			Emit(Body, *FString::Printf(TEXT("\t\t\t\tCurrentState = %d;\n"), StatementToStateIndex(FunctionContext, Statement.TargetLabel)));
@@ -1186,6 +1191,10 @@ void FKismetCppBackend::EmitGotoStatement(FKismetFunctionContext& FunctionContex
 		Emit(Body, *FString::Printf(TEXT("\t\t\t\tbreak;\n")));
 		Emit(Body, *FString::Printf(TEXT("\t\t\t}\n")));
 	}
+	else if (Statement.Type == KCST_GotoReturn)
+	{
+		Emit(Body, TEXT("\t\t\tCurrentState = -1;\n"));
+	}
 	else
 	{
 		Emit(Body, *FString::Printf(TEXT("\t\t\tCurrentState = %d;\n"), StatementToStateIndex(FunctionContext, Statement.TargetLabel)));
@@ -1195,11 +1204,13 @@ void FKismetCppBackend::EmitGotoStatement(FKismetFunctionContext& FunctionContex
 
 void FKismetCppBackend::EmitPushStateStatement(FKismetFunctionContext& FunctionContext, FBlueprintCompiledStatement& Statement)
 {
+	ensure(FunctionContext.bUseFlowStack);
 	Emit(Body, *FString::Printf(TEXT("\t\t\tStateStack.Push(%d);\n"), StatementToStateIndex(FunctionContext, Statement.TargetLabel)));
 }
 
 void FKismetCppBackend::EmitEndOfThreadStatement(FKismetFunctionContext& FunctionContext, const FString& ReturnValueString)
 {
+	ensure(FunctionContext.bUseFlowStack);
 	Emit(Body, TEXT("\t\t\tCurrentState = (StateStack.Num() > 0) ? StateStack.Pop() : -1;\n"));
 	Emit(Body, TEXT("\t\t\tbreak;\n"));
 }
@@ -1228,7 +1239,10 @@ void FKismetCppBackend::DeclareLocalVariables(FKismetFunctionContext& FunctionCo
 
 void FKismetCppBackend::DeclareStateSwitch(FKismetFunctionContext& FunctionContext)
 {
-	Emit(Body, TEXT("\tTArray< int32, TInlineAllocator<8> > StateStack;\n"));
+	if (FunctionContext.bUseFlowStack)
+	{
+		Emit(Body, TEXT("\tTArray< int32, TInlineAllocator<8> > StateStack;\n"));
+	}
 	Emit(Body, TEXT("\tint32 CurrentState = 0;\n"));
 	Emit(Body, TEXT("\tdo\n"));
 	Emit(Body, TEXT("\t{\n"));
@@ -1241,7 +1255,10 @@ void FKismetCppBackend::CloseStateSwitch(FKismetFunctionContext& FunctionContext
 {
 	// Default error-catching case 
 	Emit(Body, TEXT("\t\tdefault:\n"));
-	Emit(Body, TEXT("\t\t\tcheck(false); // Invalid state\n"));
+	if (FunctionContext.bUseFlowStack)
+	{
+		Emit(Body, TEXT("\t\t\tcheck(false); // Invalid state\n"));
+	}
 	Emit(Body, TEXT("\t\t\tbreak;\n"));
 
 	// Close the switch block and do-while loop
@@ -1354,7 +1371,39 @@ void FKismetCppBackend::ConstructFunction(FKismetFunctionContext& FunctionContex
 		// Emit local variables
 		DeclareLocalVariables(FunctionContext, LocalVariables);
 
-		DeclareStateSwitch(FunctionContext);
+		bool bUseSwitchState = false;
+		for (auto Node : FunctionContext.LinearExecutionList)
+		{
+			TArray<FBlueprintCompiledStatement*>* StatementList = FunctionContext.StatementsPerNode.Find(Node);
+			if (StatementList)
+			{
+				for (auto Statement : (*StatementList))
+				{
+					if (Statement && (
+						Statement->Type == KCST_UnconditionalGoto ||
+						Statement->Type == KCST_PushState ||
+						Statement->Type == KCST_GotoIfNot ||
+						Statement->Type == KCST_ComputedGoto ||
+						Statement->Type == KCST_EndOfThread ||
+						Statement->Type == KCST_EndOfThreadIfNot ||
+						Statement->Type == KCST_GotoReturn ||
+						Statement->Type == KCST_GotoReturnIfNot))
+					{
+						bUseSwitchState = true;
+						break;
+					}
+				}
+			}
+			if (bUseSwitchState)
+			{
+				break;
+			}
+		}
+
+		if (bUseSwitchState)
+		{
+			DeclareStateSwitch(FunctionContext);
+		}
 
 
 		// Run thru code looking only at things marked as jump targets, to make sure the jump targets are ordered in order of appearance in the linear execution list
@@ -1391,7 +1440,7 @@ void FKismetCppBackend::ConstructFunction(FKismetFunctionContext& FunctionContex
 				{
 					FBlueprintCompiledStatement& Statement = *((*StatementList)[StatementIndex]);
 
-					if (Statement.bIsJumpTarget)
+					if (Statement.bIsJumpTarget && bUseSwitchState)
 					{
 						const int32 StateNum = StatementToStateIndex(FunctionContext, &Statement);
 						Emit(Body, *FString::Printf(TEXT("\n\t\tcase %d:\n"), StateNum));
@@ -1451,6 +1500,8 @@ void FKismetCppBackend::ConstructFunction(FKismetFunctionContext& FunctionContex
 					case KCST_UnconditionalGoto:
 					case KCST_GotoIfNot:
 					case KCST_EndOfThreadIfNot:
+					case KCST_GotoReturn:
+					case KCST_GotoReturnIfNot:
 						EmitGotoStatement(FunctionContext, Statement);
 						break;
 					case KCST_PushState:
@@ -1477,7 +1528,10 @@ void FKismetCppBackend::ConstructFunction(FKismetFunctionContext& FunctionContex
 			}
 		}
 
-		CloseStateSwitch(FunctionContext);
+		if (bUseSwitchState)
+		{
+			CloseStateSwitch(FunctionContext);
+		}
 	}
 
 	EmitReturnStatement(FunctionContext, ReturnValueString);
diff --git a/Engine/Source/Editor/KismetCompiler/Private/KismetCompilerMisc.cpp b/Engine/Source/Editor/KismetCompiler/Private/KismetCompilerMisc.cpp
--- a/Engine/Source/Editor/KismetCompiler/Private/KismetCompilerMisc.cpp
+++ b/Engine/Source/Editor/KismetCompiler/Private/KismetCompilerMisc.cpp
@@ -1006,6 +1006,7 @@ FKismetFunctionContext::FKismetFunctionContext(FCompilerResultsLog& InMessageLog
 	, NetFlags(0)
 	, SourceEventFromStubGraph(NULL)
 	, bGeneratingCpp(bInGeneratingCpp)
+	, bUseFlowStack(true)
 {
 	NetNameMap = new FNetNameMapping();
 	bAllocatedNetNameMap = true;
@@ -1051,6 +1052,243 @@ int32 FKismetFunctionContext::GetContextUniqueID()
 	return UUIDCounter++;
 }
 
+void FKismetFunctionContext::MergeAdjacentStates()
+{
+	for (int32 ExecIndex = 0; ExecIndex < LinearExecutionList.Num(); ++ExecIndex)
+	{
+		// if the last statement in current node jumps to the first statement in next node, then it's redundant
+		const auto CurrentNode = LinearExecutionList[ExecIndex];
+		auto CurStatementList = StatementsPerNode.Find(CurrentNode);
+		const bool CurrentNodeIsValid = CurrentNode && CurStatementList && CurStatementList->Num();
+		const auto LastStatementInCurrentNode = CurrentNodeIsValid ? CurStatementList->Last() : NULL;
+
+		if (LastStatementInCurrentNode
+			&& LastStatementInCurrentNode->TargetLabel
+			&& (LastStatementInCurrentNode->Type == KCST_UnconditionalGoto)
+			&& !LastStatementInCurrentNode->bIsJumpTarget)
+		{
+			const auto NextNodeIndex = ExecIndex + 1;
+			const auto NextNode = LinearExecutionList.IsValidIndex(NextNodeIndex) ? LinearExecutionList[NextNodeIndex] : NULL;
+			const auto NextNodeStatements = StatementsPerNode.Find(NextNode);
+			const bool bNextNodeValid = NextNode && NextNodeStatements && NextNodeStatements->Num();
+			const auto FirstStatementInNextNode = bNextNodeValid ? (*NextNodeStatements)[0] : NULL;
+			if (FirstStatementInNextNode == LastStatementInCurrentNode->TargetLabel)
+			{
+				CurStatementList->RemoveAt(CurStatementList->Num() - 1);
+			}
+		}
+	}
+
+	// Remove unnecessary GotoReturn statements
+	// if it's last statement generated by last node (in LinearExecution) then it can be removed
+	const auto LastExecutedNode = LinearExecutionList.Num() ? LinearExecutionList.Last() : NULL;
+	TArray<FBlueprintCompiledStatement*>* StatementList = StatementsPerNode.Find(LastExecutedNode);
+	FBlueprintCompiledStatement* LastStatementInLastNode = (StatementList && StatementList->Num()) ? StatementList->Last() : NULL;
+	if (LastStatementInLastNode && (KCST_GotoReturn == LastStatementInLastNode->Type) && !LastStatementInLastNode->bIsJumpTarget)
+	{
+		StatementList->RemoveAt(StatementList->Num() - 1);
+	}
+}
+
+struct FGotoMapUtils
+{
+	static bool IsUberGraphEventStatement(const FBlueprintCompiledStatement* GotoStatement)
+	{
+		return GotoStatement
+			&& (GotoStatement->Type == KCST_CallFunction) 
+			&& (GotoStatement->UbergraphCallIndex == 0);
+	}
+
+	static UEdGraphNode* TargetNodeFromPin(const FBlueprintCompiledStatement* GotoStatement, const UEdGraphPin* ExecNet)
+	{
+		UEdGraphNode* TargetNode = NULL;
+		if (ExecNet && GotoStatement)
+		{
+			if (IsUberGraphEventStatement(GotoStatement))
+			{
+				TargetNode = ExecNet->GetOwningNode();
+			}
+			else if (ExecNet->LinkedTo.Num() > 0)
+			{
+				TargetNode = ExecNet->LinkedTo[0]->GetOwningNode();
+			}
+		}
+		return TargetNode;
+	}
+
+	static UEdGraphNode* TargetNodeFromMap(const FBlueprintCompiledStatement* GotoStatement, const TMap< FBlueprintCompiledStatement*, UEdGraphPin* >& GotoFixupRequestMap)
+	{
+		auto ExecNetPtr = GotoFixupRequestMap.Find(GotoStatement);
+		auto ExecNet = ExecNetPtr ? *ExecNetPtr : NULL;
+		return TargetNodeFromPin(GotoStatement, ExecNet);
+	}
+};
+
+void FKismetFunctionContext::ResolveGotoFixups()
+{
+	if (bCreateDebugData)
+	{
+		// if we're debugging, go through an insert a wire trace before  
+		// every "goto" statement so we can trace what execution pin a node
+		// was executed from
+		for (auto GotoIt = GotoFixupRequestMap.CreateIterator(); GotoIt; ++GotoIt)
+		{
+			FBlueprintCompiledStatement* GotoStatement = GotoIt.Key();
+			if (FGotoMapUtils::IsUberGraphEventStatement(GotoStatement))
+			{
+				continue;
+			}
+
+			InsertWireTrace(GotoIt.Key(), GotoIt.Value());
+		}
+	}
+
+	// Resolve the remaining fixups
+	for (auto GotoIt = GotoFixupRequestMap.CreateIterator(); GotoIt; ++GotoIt)
+	{
+		FBlueprintCompiledStatement* GotoStatement = GotoIt.Key();
+		const UEdGraphPin* ExecNet = GotoIt.Value();
+		const UEdGraphNode* TargetNode = FGotoMapUtils::TargetNodeFromPin(GotoStatement, ExecNet);
+
+		if (TargetNode == NULL)
+		{
+			// If Execution Flow Stack isn't necessary, then use GotoReturn instead EndOfThread.
+			// EndOfThread pops Execution Flow Stack, GotoReturn dosen't.
+			GotoStatement->Type = bUseFlowStack
+				? ((GotoStatement->Type == KCST_GotoIfNot) ? KCST_EndOfThreadIfNot : KCST_EndOfThread)
+				: ((GotoStatement->Type == KCST_GotoIfNot) ? KCST_GotoReturnIfNot : KCST_GotoReturn);
+		}
+		else
+		{
+			// Try to resolve the goto
+			TArray<FBlueprintCompiledStatement*>* StatementList = StatementsPerNode.Find(TargetNode);
+
+			if ((StatementList == NULL) || (StatementList->Num() == 0))
+			{
+				MessageLog.Error(TEXT("Statement tried to pass control flow to a node @@ that generates no code"), TargetNode);
+				GotoStatement->Type = KCST_CompileError;
+			}
+			else
+			{
+				// Wire up the jump target and notify the target that it is targeted
+				FBlueprintCompiledStatement& FirstStatement = *((*StatementList)[0]);
+				GotoStatement->TargetLabel = &FirstStatement;
+				FirstStatement.bIsJumpTarget = true;
+			}
+		}
+	}
+
+	// Clear out the pending fixup map
+	GotoFixupRequestMap.Empty();
+
+	//@TODO: Remove any wire debug sites where the next statement is a stack pop
+}
+
+void FKismetFunctionContext::FinalSortLinearExecList()
+{
+	auto K2Schema = CastChecked<UEdGraphSchema_K2>(Schema);
+	LinearExecutionList.RemoveAllSwap([&](UEdGraphNode* CurrentNode)
+	{
+		auto CurStatementList = StatementsPerNode.Find(CurrentNode);
+		return !(CurrentNode && CurStatementList && CurStatementList->Num());
+	});
+
+	TSet<UEdGraphNode*> UnsortedExecutionSet(LinearExecutionList);
+	LinearExecutionList.Empty();
+	TArray<UEdGraphNode*> SortedLinearExecutionList;
+
+	check(EntryPoint);
+	SortedLinearExecutionList.Push(EntryPoint);
+	UnsortedExecutionSet.Remove(EntryPoint);
+
+	TSet<UEdGraphNode*> NodesToStartNextChain;
+
+	while (UnsortedExecutionSet.Num())
+	{
+		UEdGraphNode* NextNode = NULL;
+
+		// get last state target
+		const auto CurrentNode = SortedLinearExecutionList.Last();
+		const auto CurStatementList = StatementsPerNode.Find(CurrentNode);
+		const bool CurrentNodeIsValid = CurrentNode && CurStatementList && CurStatementList->Num();
+		const auto LastStatementInCurrentNode = CurrentNodeIsValid ? CurStatementList->Last() : NULL;
+
+		// Find next element in current chain
+		if (LastStatementInCurrentNode && (LastStatementInCurrentNode->Type == KCST_UnconditionalGoto))
+		{
+			auto TargetNode = FGotoMapUtils::TargetNodeFromMap(LastStatementInCurrentNode, GotoFixupRequestMap);
+			NextNode = UnsortedExecutionSet.Remove(TargetNode) ? TargetNode : NULL;
+		}
+
+		if (CurrentNode)
+		{
+			for (auto Pin : CurrentNode->Pins)
+			{
+				if (Pin && (EEdGraphPinDirection::EGPD_Output == Pin->Direction) && K2Schema->IsExecPin(*Pin) && Pin->LinkedTo.Num())
+				{
+					for (auto Link : Pin->LinkedTo)
+					{
+						auto LinkedNode = Link->GetOwningNodeUnchecked();
+						if (LinkedNode && (LinkedNode != NextNode) && UnsortedExecutionSet.Contains(LinkedNode))
+						{
+							NodesToStartNextChain.Add(LinkedNode);
+						}
+					}
+				}
+			}
+		}
+
+		// Start next chain if the current is done
+		while (NodesToStartNextChain.Num() && !NextNode)
+		{
+			auto Iter = NodesToStartNextChain.CreateIterator();
+			NextNode = UnsortedExecutionSet.Remove(*Iter) ? *Iter : NULL;
+			Iter.RemoveCurrent();
+		}
+
+		if (!NextNode)
+		{
+			auto Iter = UnsortedExecutionSet.CreateIterator();
+			NextNode = *Iter;
+			Iter.RemoveCurrent();
+		}
+
+		check(NextNode);
+		SortedLinearExecutionList.Push(NextNode);
+	}
+
+	LinearExecutionList = SortedLinearExecutionList;
+}
+
+void FKismetFunctionContext::ResolveStatements()
+{
+	BP_SCOPED_COMPILER_EVENT_STAT(EKismetCompilerStats_ResolveCompiledStatements);
+	FinalSortLinearExecList();
+
+	static const FBoolConfigValueHelper OptimizeExecutionFlowStack(TEXT("Kismet"), TEXT("bOptimizeExecutionFlowStack"), GEngineIni);
+	if (OptimizeExecutionFlowStack)
+	{
+		const bool bFlowStackIsRequired = AllGeneratedStatements.ContainsByPredicate(
+			[](const FBlueprintCompiledStatement* Statement)
+		{
+			return Statement && (
+				(Statement->Type == KCST_EndOfThreadIfNot) ||
+				(Statement->Type == KCST_EndOfThread) ||
+				(Statement->Type == KCST_PushState));
+		});
+
+		bUseFlowStack = bFlowStackIsRequired;
+	}
+
+	ResolveGotoFixups();
+
+	static const FBoolConfigValueHelper OptimizeAdjacentStates(TEXT("Kismet"), TEXT("bOptimizeAdjacentStates"), GEngineIni);
+	if (OptimizeAdjacentStates)
+	{
+		MergeAdjacentStates();
+	}
+}
+
 struct FEventGraphUtils
 {
 	static bool IsEntryPointNode(const UK2Node* Node)
diff --git a/Engine/Source/Editor/KismetCompiler/Private/KismetCompilerVMBackend.cpp b/Engine/Source/Editor/KismetCompiler/Private/KismetCompilerVMBackend.cpp
--- a/Engine/Source/Editor/KismetCompiler/Private/KismetCompilerVMBackend.cpp
+++ b/Engine/Source/Editor/KismetCompiler/Private/KismetCompilerVMBackend.cpp
@@ -208,6 +208,8 @@ private:
 	
 	// Is this compiling the ubergraph?
 	bool bIsUbergraph;
+
+	FBlueprintCompiledStatement& ReturnStatement;
 protected:
 	/**
 	 * This class is designed to be used like so to emit a bytecode context expression:
@@ -308,12 +310,13 @@ protected:
 		}
 	};
 public:
-	FScriptBuilderBase(TArray<uint8>& InScript, UBlueprintGeneratedClass* InClass, UEdGraphSchema_K2* InSchema, FKismetCompilerVMBackend::TStatementToSkipSizeMap& InUbergraphStatementLabelMap, bool bInIsUbergraph)
+	FScriptBuilderBase(TArray<uint8>& InScript, UBlueprintGeneratedClass* InClass, UEdGraphSchema_K2* InSchema, FKismetCompilerVMBackend::TStatementToSkipSizeMap& InUbergraphStatementLabelMap, bool bInIsUbergraph, FBlueprintCompiledStatement& InReturnStatement)
 		: Writer(InScript)
 		, ClassBeingBuilt(InClass)
 		, Schema(InSchema)
 		, UbergraphStatementLabelMap(InUbergraphStatementLabelMap)
 		, bIsUbergraph(bInIsUbergraph)
+		, ReturnStatement(InReturnStatement)
 	{
 		VectorStruct = FindObjectChecked<UScriptStruct>(UObject::StaticClass(), TEXT("Vector"));
 		RotatorStruct = FindObjectChecked<UScriptStruct>(UObject::StaticClass(), TEXT("Rotator"));
@@ -1073,7 +1076,7 @@ public:
 			// Now include the boolean expression
 			EmitTerm(Statement.LHS, (UProperty*)(GetDefault<UBoolProperty>()));
 		}
-		else
+		else if (Statement.Type == KCST_UnconditionalGoto)
 		{
 			// Emit the jump with a dummy address
 			Writer << EX_Jump;
@@ -1082,6 +1085,31 @@ public:
 			// Queue up a fixup to be done once all label offsets are known
 			JumpTargetFixupMap.Add(PatchUpNeededAtOffset, Statement.TargetLabel);
 		}
+		else if (Statement.Type == KCST_GotoReturn)
+		{
+			// Emit the jump with a dummy address
+			Writer << EX_Jump;
+			CodeSkipSizeType PatchUpNeededAtOffset = Writer.EmitPlaceholderSkip();
+
+			// Queue up a fixup to be done once all label offsets are known
+			JumpTargetFixupMap.Add(PatchUpNeededAtOffset, &ReturnStatement);
+		}
+		else if (Statement.Type == KCST_GotoReturnIfNot)
+		{
+			// Emit the jump with a dummy address
+			Writer << EX_JumpIfNot;
+			CodeSkipSizeType PatchUpNeededAtOffset = Writer.EmitPlaceholderSkip();
+
+			// Queue up a fixup to be done once all label offsets are known
+			JumpTargetFixupMap.Add(PatchUpNeededAtOffset, &ReturnStatement);
+
+			// Now include the boolean expression
+			EmitTerm(Statement.LHS, (UProperty*)(GetDefault<UBoolProperty>()));
+		}
+		else
+		{
+			ensureMsg(false, TEXT("FScriptBuilderBase::EmitGoto unknown type"));
+		}
 	}
 
 	void EmitPushExecState(FBlueprintCompiledStatement& Statement)
@@ -1249,6 +1277,8 @@ public:
 		case KCST_UnconditionalGoto:
 		case KCST_GotoIfNot:
 		case KCST_EndOfThreadIfNot:
+		case KCST_GotoReturn:
+		case KCST_GotoReturnIfNot:
 			EmitGoto(Statement);
 			break;
 		case KCST_PushState:
@@ -1312,16 +1342,19 @@ void FKismetCompilerVMBackend::ConstructFunction(FKismetFunctionContext& Functio
 
 	TArray<uint8>& ScriptArray = Function->Script;
 
-	FScriptBuilderBase ScriptWriter(ScriptArray, Class, Schema, UbergraphStatementLabelMap, bIsUbergraph);
-
-	// Since the flow stack always assumes there is something to pop, the first pushed item should be the return block for the function
+	// Return statement, to push on FlowStack or to use with _GotoReturn
 	FBlueprintCompiledStatement ReturnStatement;
 	ReturnStatement.Type = KCST_Return;
 
+	FScriptBuilderBase ScriptWriter(ScriptArray, Class, Schema, UbergraphStatementLabelMap, bIsUbergraph, ReturnStatement);
+
 	if (!bGenerateStubOnly)
 	{
 		ReturnStatement.bIsJumpTarget = true;
-		ScriptWriter.PushReturnAddress(ReturnStatement);
+		if (FunctionContext.bUseFlowStack)
+		{
+			ScriptWriter.PushReturnAddress(ReturnStatement);
+		}
 	
 		// Emit code in the order specified by the linear execution list (the first node is always the entry point for the function)
 		for (int32 NodeIndex = 0; NodeIndex < FunctionContext.LinearExecutionList.Num(); ++NodeIndex)
diff --git a/Engine/Source/Editor/KismetCompiler/Public/BlueprintCompiledStatement.h b/Engine/Source/Editor/KismetCompiler/Public/BlueprintCompiledStatement.h
--- a/Engine/Source/Editor/KismetCompiler/Public/BlueprintCompiledStatement.h
+++ b/Engine/Source/Editor/KismetCompiler/Public/BlueprintCompiledStatement.h
@@ -34,6 +34,8 @@ enum EKismetCompiledStatementType
 	KCST_MetaCast = 24,	    // Cast<TargetClass>(TargetObject)
 	KCST_AssignmentOnPersistentFrame = 25, //
 	KCST_CastInterfaceToObj = 26, // Cast<TargetClass>(TargetInterface)
+	KCST_GotoReturn = 27,	// goto ReturnLabel
+	KCST_GotoReturnIfNot = 28, // [if (!TargetObject->TargetProperty)] goto TargetLabel
 };
 
 //@TODO: Too rigid / icky design
diff --git a/Engine/Source/Editor/KismetCompiler/Public/EdGraphCompilerUtilities.h b/Engine/Source/Editor/KismetCompiler/Public/EdGraphCompilerUtilities.h
--- a/Engine/Source/Editor/KismetCompiler/Public/EdGraphCompilerUtilities.h
+++ b/Engine/Source/Editor/KismetCompiler/Public/EdGraphCompilerUtilities.h
@@ -28,6 +28,7 @@ DECLARE_CYCLE_STAT_EXTERN(TEXT("Clean and Sanitize Class"), EKismetCompilerStats
 DECLARE_CYCLE_STAT_EXTERN(TEXT("Create Class Properties"), EKismetCompilerStats_CreateClassVariables, STATGROUP_KismetCompiler, );
 DECLARE_CYCLE_STAT_EXTERN(TEXT("Bind and Link Class"), EKismetCompilerStats_BindAndLinkClass, STATGROUP_KismetCompiler, );
 DECLARE_CYCLE_STAT_EXTERN(TEXT("Calculate checksum of CDO"), EKismetCompilerStats_ChecksumCDO, STATGROUP_KismetCompiler, );
+DECLARE_CYCLE_STAT_EXTERN(TEXT("Resolve compiled statements"), EKismetCompilerStats_ResolveCompiledStatements, STATGROUP_KismetCompiler, );
 
 //////////////////////////////////////////////////////////////////////////
 
diff --git a/Engine/Source/Editor/KismetCompiler/Public/KismetCompiledFunctionContext.h b/Engine/Source/Editor/KismetCompiler/Public/KismetCompiledFunctionContext.h
--- a/Engine/Source/Editor/KismetCompiler/Public/KismetCompiledFunctionContext.h
+++ b/Engine/Source/Editor/KismetCompiler/Public/KismetCompiledFunctionContext.h
@@ -12,9 +12,6 @@ namespace KismetCompilerDebugOptions
 	//@TODO: K2: Turning this off is probably broken due to state merging not working with the current code generation
 	enum { DebuggingCompiler = 1 };
 
-	// Should adjacent states be merged together?
-	enum { MergeAdjacentStates = !DebuggingCompiler };
-
 	// Should the compiler emit node comments to the backends?
 	enum { EmitNodeComments = DebuggingCompiler };
 }
@@ -94,6 +91,9 @@ public:
 	TArray< UK2Node_CallFunction* > LatentFunctionCalls;
 
 	bool bGeneratingCpp;
+
+	//Does this function use requires FlowStack ?
+	bool bUseFlowStack;
 public:
 	FKismetFunctionContext(FCompilerResultsLog& InMessageLog, UEdGraphSchema_K2* InSchema, UBlueprintGeneratedClass* InNewClass, UBlueprint* InBlueprint, bool bInGeneratingCpp);
 	
@@ -273,90 +273,20 @@ public:
 		return (SourceStatementList != NULL) && (SourceStatementList->Num() > 0);
 	}
 
-	// Sorts the 'linear execution list' again by likely execution order; the list should only contain impure nodes by this point.
-	void FinalSortLinearExecList()
-	{
-		// By this point, all of the pure data dependencies have been taken care of and the position in the list 
-		// doesn't represent much.
-		// Doing this sort will enable more gotos to be optimized away, and make the C++ backend generated code easier to read
-		for (int32 TestIndex = 0; TestIndex < LinearExecutionList.Num(); )
-		{
-			UEdGraphNode* CurrentNode = LinearExecutionList[TestIndex];
-			TArray<FBlueprintCompiledStatement*>* CurStatementList = StatementsPerNode.Find(CurrentNode);
-
-			if ((CurStatementList == NULL) || (CurStatementList->Num() == 0))
-			{
-				// Node generated nothing for the backend to consume, remove it
-				LinearExecutionList.RemoveAtSwap(TestIndex);
-			}
-			else
-			{
-				//@TODO: Actually do some real sorting here!
-				//FBlueprintCompiledStatement& LastStatementInCurrentNode = *(CurStatementList->Last());
-				//UEdGraphNode* JumpTargetNode = GotoFixupRequestMap.FindChecked(&LastStatementInCurrentNode);
-
-				// Sort so the entry point node is always first
-				if ((CurrentNode == EntryPoint) && (TestIndex > 0))
-				{
-					LinearExecutionList.SwapMemory(0, TestIndex);
-				}
-				else
-				{
-					TestIndex++;
-				}
-			}
-		}
-	}
-
+private:
 	// Optimize out any useless jumps (jump to the very next statement, where the control flow can just fall through)
-	void MergeAdjacentStates()
-	{
-		for (int32 i = 0; i < LinearExecutionList.Num(); ++i)
-		{
-			UEdGraphNode* CurrentNode = LinearExecutionList[i];
-			TArray<FBlueprintCompiledStatement*>* CurStatementList = StatementsPerNode.Find(CurrentNode);
+	void MergeAdjacentStates();
 
-			if ((CurStatementList != NULL) && (CurStatementList->Num() > 0))
-			{
-				FBlueprintCompiledStatement& LastStatementInCurrentNode = *(CurStatementList->Last());
-				if (LastStatementInCurrentNode.Type == KCST_UnconditionalGoto)
-				{
-					// Last statement is an unconditional goto; see if we can remove it
+	// Sorts the 'linear execution list' again by likely execution order; the list should only contain impure nodes by this point.
+	void FinalSortLinearExecList();
 
-					// Find the next node that generated statements
-					UEdGraphNode* NextNode = NULL;
-					for (int32 j = i+1; j < LinearExecutionList.Num(); ++j)
-					{
-						UEdGraphNode* TestNode = LinearExecutionList[j];
-						if (StatementsPerNode.Find(TestNode) != NULL)
-						{
-							NextNode = TestNode;
-							break;
-						}
-					}
+	/** Resolves all pending goto fixups; Should only be called after all nodes have had a chance to generate code! */
+	void ResolveGotoFixups();
 
-					if (NextNode != NULL)
-					{
-						UEdGraphNode* JumpTargetNode = NULL;
-
-						UEdGraphPin const* JumpPin = GotoFixupRequestMap.FindChecked(&LastStatementInCurrentNode);
-						if ((JumpPin != NULL) && (JumpPin->LinkedTo.Num() > 0))
-						{
-							JumpTargetNode = JumpPin->LinkedTo[0]->GetOwningNode();
-						}
-
-						// Is this node the target of the last statement in the previous node?
-						if (NextNode == JumpTargetNode)
-						{
-							// Jump and label are adjacent, can optimize it out
-							GotoFixupRequestMap.Remove(&LastStatementInCurrentNode);
-							CurStatementList->RemoveAt(CurStatementList->Num() - 1);
-						}
-					}
-				}
-			}
-		}
-	}
+public:
+
+	// The function links gotos, sorts statments, and merges adjacent ones. 
+	void ResolveStatements();
 
 	/**
 	 * Makes sure an KCST_WireTraceSite is inserted before the specified 
@@ -406,98 +336,6 @@ public:
 			}
 		}
 	}
-
-	/** Resolves all pending goto fixups; Should only be called after all nodes have had a chance to generate code! */
-	void ResolveGotoFixups()
-	{
-		/**
-		 * This local lambda is intended to save on code duplication within this
-		 * function. It tests the supplied statement to see if it is an event 
-		 * entry point.
-		 *
-		 * @param GotoStatement		The compiled statement you wish to test.
-		 * @return True if the passed statement is an event entry point, false if not.
-		 */
-		auto IsUberGraphEventStatement = [](FBlueprintCompiledStatement* GotoStatement) -> bool 
-		{
-			return ((GotoStatement->Type == KCST_CallFunction) && (GotoStatement->UbergraphCallIndex == 0));
-		};
-
-		if (bCreateDebugData)
-		{
-			// if we're debugging, go through an insert a wire trace before  
-			// every "goto" statement so we can trace what execution pin a node
-			// was executed from
-			for (auto GotoIt = GotoFixupRequestMap.CreateIterator(); GotoIt; ++GotoIt)
-			{
-				FBlueprintCompiledStatement* GotoStatement = GotoIt.Key();
-				if (IsUberGraphEventStatement(GotoStatement))
-				{
-					continue;
-				}
-
-				InsertWireTrace(GotoIt.Key(), GotoIt.Value());
-			}
-		}
-
-		if (KismetCompilerDebugOptions::MergeAdjacentStates)
-		{
-			MergeAdjacentStates();
-		}
-
-		// Resolve the remaining fixups
-		for (auto GotoIt = GotoFixupRequestMap.CreateIterator(); GotoIt; ++GotoIt)
-		{
-			FBlueprintCompiledStatement* GotoStatement = GotoIt.Key();
-
-			UEdGraphPin const* ExecNet = GotoIt.Value();
-			if (ExecNet == NULL)
-			{
-				// Execution flow ended here; pop the stack
-				GotoStatement->Type = (GotoStatement->Type == KCST_GotoIfNot) ? KCST_EndOfThreadIfNot : KCST_EndOfThread;
-				continue;
-			}
-
-			UEdGraphNode* TargetNode = NULL;
-			if (IsUberGraphEventStatement(GotoStatement))
-			{
-				TargetNode = ExecNet->GetOwningNode();
-			}
-			else if (ExecNet->LinkedTo.Num() > 0)
-			{
-				TargetNode = ExecNet->LinkedTo[0]->GetOwningNode();
-			}
-
-			if (TargetNode == NULL)
-			{
-				// Execution flow ended here; pop the stack
-				GotoStatement->Type = (GotoStatement->Type == KCST_GotoIfNot) ? KCST_EndOfThreadIfNot : KCST_EndOfThread;
-			}
-			else
-			{
-				// Try to resolve the goto
-				TArray<FBlueprintCompiledStatement*>* StatementList = StatementsPerNode.Find(TargetNode);
-
-				if ((StatementList == NULL) || (StatementList->Num() == 0))
-				{
-					MessageLog.Error(TEXT("Statement tried to pass control flow to a node @@ that generates no code"), TargetNode);
-					GotoStatement->Type = KCST_CompileError;
-				}
-				else
-				{
-					// Wire up the jump target and notify the target that it is targeted
-					FBlueprintCompiledStatement& FirstStatement = *((*StatementList)[0]);
-					GotoStatement->TargetLabel = &FirstStatement;
-					FirstStatement.bIsJumpTarget = true;
-				}
-			}
-		}
-
-		// Clear out the pending fixup map
-		GotoFixupRequestMap.Empty();
-
-		//@TODO: Remove any wire debug sites where the next statement is a stack pop
-	}
 	
 	/** Looks for a pin of the given name, erroring if the pin is not found or if the direction doesn't match (doesn't verify the pin type) */
 	UEdGraphPin* FindRequiredPinByName(const UEdGraphNode* Node, const FString& PinName, EEdGraphPinDirection RequiredDirection = EGPD_MAX)
