Merging PR #680 to fix crashes when casting an Interface pin to an Object pin inside of Blueprints.

diff --git a/Engine/Source/Editor/BlueprintGraph/Private/CallFunctionHandler.cpp b/Engine/Source/Editor/BlueprintGraph/Private/CallFunctionHandler.cpp
--- a/Engine/Source/Editor/BlueprintGraph/Private/CallFunctionHandler.cpp
+++ b/Engine/Source/Editor/BlueprintGraph/Private/CallFunctionHandler.cpp
@@ -234,7 +234,7 @@ void FKCHandler_CallFunction::CreateFunctionCallStatement(FKismetFunctionContext
 									// insert a cast op before a call to the function (and replace
 									// the param with the result from the cast)
 									FBlueprintCompiledStatement& CastStatement = Context.AppendStatementForNode(Node);
-									CastStatement.Type = KCST_CastObjToInterface;
+									CastStatement.Type = InterfaceClass->HasAnyClassFlags(CLASS_Interface) ? KCST_CastObjToInterface : KCST_CastInterfaceToObj;
 									CastStatement.LHS = *InterfaceTerm;
 									CastStatement.RHS.Add(ClassTerm);
 									CastStatement.RHS.Add(*Term);
@@ -604,7 +604,8 @@ void FKCHandler_CallFunction::RegisterNets(FKismetFunctionContext& Context, UEdG
 		// if we have an object plugged into an interface pin, let's create a 
 		// term that'll be used as an intermediate, holding the result of a cast 
 		// from object to interface
-		if ((Pin->PinType.PinCategory == K2Schema->PC_Interface) && (Pin->LinkedTo[0]->PinType.PinCategory == K2Schema->PC_Object))
+		if (((Pin->PinType.PinCategory == K2Schema->PC_Interface) && (Pin->LinkedTo[0]->PinType.PinCategory == K2Schema->PC_Object)) ||
+			((Pin->PinType.PinCategory == K2Schema->PC_Object) && (Pin->LinkedTo[0]->PinType.PinCategory == K2Schema->PC_Interface)))
 		{
 			FBPTerminal* InterfaceTerm = Context.CreateLocalTerminal();
 			InterfaceTerm->CopyFromPin(Pin, Context.NetNameMap->MakeValidName(Pin) + TEXT("_CastInput"));
diff --git a/Engine/Source/Editor/BlueprintGraph/Private/MakeStructHandler.cpp b/Engine/Source/Editor/BlueprintGraph/Private/MakeStructHandler.cpp
--- a/Engine/Source/Editor/BlueprintGraph/Private/MakeStructHandler.cpp
+++ b/Engine/Source/Editor/BlueprintGraph/Private/MakeStructHandler.cpp
@@ -107,10 +107,7 @@ void FKCHandler_MakeStruct::Compile(FKismetFunctionContext& Context, UEdGraphNod
 			DstTerm->AssociatedVarProperty = BoundProperty;
 			DstTerm->Context = OutputStructTerm;
 
-			FBlueprintCompiledStatement& Statement = Context.AppendStatementForNode(Node);
-			Statement.Type = KCST_Assignment;
-			Statement.LHS = DstTerm;
-			Statement.RHS.Add(SrcTerm);
+			FKismetCompilerUtilities::CreateObjectAssignmentStatement(Context, Node, SrcTerm, DstTerm);
 		}
 	}
 	if (!Node->IsNodePure())
diff --git a/Engine/Source/Editor/BlueprintGraph/Private/VariableSetHandler.cpp b/Engine/Source/Editor/BlueprintGraph/Private/VariableSetHandler.cpp
--- a/Engine/Source/Editor/BlueprintGraph/Private/VariableSetHandler.cpp
+++ b/Engine/Source/Editor/BlueprintGraph/Private/VariableSetHandler.cpp
@@ -74,11 +74,7 @@ void FKCHandler_VariableSet::InnerAssignment(FKismetFunctionContext& Context, UE
 
 	if ((VariableTerm != NULL) && (ValueTerm != NULL))
 	{
-		FBlueprintCompiledStatement& Statement = Context.AppendStatementForNode(Node);
-
-		Statement.Type = KCST_Assignment;
-		Statement.LHS = *VariableTerm;
-		Statement.RHS.Add(*ValueTerm);
+		FKismetCompilerUtilities::CreateObjectAssignmentStatement(Context, Node, *ValueTerm, *VariableTerm);
 
 		if (!(*VariableTerm)->IsTermWritable())
 		{
diff --git a/Engine/Source/Editor/KismetCompiler/Private/KismetCompilerMisc.cpp b/Engine/Source/Editor/KismetCompiler/Private/KismetCompilerMisc.cpp
--- a/Engine/Source/Editor/KismetCompiler/Private/KismetCompilerMisc.cpp
+++ b/Engine/Source/Editor/KismetCompiler/Private/KismetCompilerMisc.cpp
@@ -601,6 +601,39 @@ UEdGraphPin* FKismetCompilerUtilities::GenerateAssignmentNodes(class FKismetComp
 	return LastThen;
 }
 
+void FKismetCompilerUtilities::CreateObjectAssignmentStatement(FKismetFunctionContext& Context, UEdGraphNode* Node, FBPTerminal* SrcTerm, FBPTerminal* DstTerm)
+{
+	UClass* InputObjClass = Cast<UClass>(SrcTerm->Type.PinSubCategoryObject.Get());
+	UClass* OutputObjClass = Cast<UClass>(DstTerm->Type.PinSubCategoryObject.Get());
+
+	const bool bIsOutputInterface = ((OutputObjClass != NULL) && OutputObjClass->HasAnyClassFlags(CLASS_Interface));
+	const bool bIsInputInterface = ((InputObjClass != NULL) && InputObjClass->HasAnyClassFlags(CLASS_Interface));
+
+	if (bIsOutputInterface != bIsInputInterface)
+	{
+		// Create a literal term from the class specified in the node
+		FBPTerminal* ClassTerm = Context.CreateLocalTerminal(ETerminalSpecification::TS_Literal);
+		ClassTerm->Name = OutputObjClass->GetName();
+		ClassTerm->bIsLiteral = true;
+		ClassTerm->Source = DstTerm->Source;
+		ClassTerm->ObjectLiteral = OutputObjClass;
+
+		EKismetCompiledStatementType CastOpType = bIsOutputInterface ? KCST_CastObjToInterface : KCST_CastInterfaceToObj;
+		FBlueprintCompiledStatement& CastStatement = Context.AppendStatementForNode(Node);
+		CastStatement.Type = CastOpType;
+		CastStatement.LHS = DstTerm;
+		CastStatement.RHS.Add(ClassTerm);
+		CastStatement.RHS.Add(SrcTerm);
+	}
+	else
+	{
+		FBlueprintCompiledStatement& Statement = Context.AppendStatementForNode(Node);
+		Statement.Type = KCST_Assignment;
+		Statement.LHS = DstTerm;
+		Statement.RHS.Add(SrcTerm);
+	}
+}
+
 /** Creates a property named PropertyName of type PropertyType in the Scope or returns NULL if the type is unknown, but does *not* link that property in */
 UProperty* FKismetCompilerUtilities::CreatePropertyOnScope(UStruct* Scope, const FName& PropertyName, const FEdGraphPinType& Type, UClass* SelfClass, uint64 PropertyFlags, const UEdGraphSchema_K2* Schema, FCompilerResultsLog& MessageLog)
 {
diff --git a/Engine/Source/Editor/KismetCompiler/Public/KismetCompilerMisc.h b/Engine/Source/Editor/KismetCompiler/Public/KismetCompilerMisc.h
--- a/Engine/Source/Editor/KismetCompiler/Public/KismetCompilerMisc.h
+++ b/Engine/Source/Editor/KismetCompiler/Public/KismetCompilerMisc.h
@@ -78,6 +78,9 @@ public:
 
 	/** Create 'set var by name' nodes and hook them up - used to set values when components are added or actor are created at run time. Returns the 'last then' pin of the assignment nodes */
 	static UEdGraphPin* GenerateAssignmentNodes( class FKismetCompilerContext& CompilerContext, UEdGraph* SourceGraph, UK2Node_CallFunction* CallBeginSpawnNode, UEdGraphNode* SpawnNode, UEdGraphPin* CallBeginResult, const UClass* ForClass );
+
+	/** Create Kismet assignment statement with proper object <-> interface cast */
+	static void CreateObjectAssignmentStatement(FKismetFunctionContext& Context, UEdGraphNode* Node, FBPTerminal* SrcTerm, FBPTerminal* DstTerm);
 };
 
 //////////////////////////////////////////////////////////////////////////
diff --git a/Engine/Source/Runtime/CoreUObject/Private/UObject/ScriptCore.cpp b/Engine/Source/Runtime/CoreUObject/Private/UObject/ScriptCore.cpp
--- a/Engine/Source/Runtime/CoreUObject/Private/UObject/ScriptCore.cpp
+++ b/Engine/Source/Runtime/CoreUObject/Private/UObject/ScriptCore.cpp
@@ -2026,7 +2026,7 @@ void UObject::execInterfaceToObject(FFrame& Stack, RESULT_DECL)
 	Stack.Step(Stack.Object, &InterfaceInput);
 
 	UObject* InputObjWithInterface = InterfaceInput.GetObjectRef();
-	if (InputObjWithInterface->IsA(ObjClassToCastTo))
+	if (InputObjWithInterface && InputObjWithInterface->IsA(ObjClassToCastTo))
 	{
 		*(UObject**)Result = InputObjWithInterface;
 	}
