Blueprint Skeleton classes inherit their parent's Blueprint Skeleton class instead of the generated class by default. Removed ini flag to toggle this functionality.

diff --git a/Engine/Config/BaseEngine.ini b/Engine/Config/BaseEngine.ini
--- a/Engine/Config/BaseEngine.ini
+++ b/Engine/Config/BaseEngine.ini
@@ -1473,7 +1473,6 @@ bReinstanceOnlyWhenNecessary=true
 bChangeDefaultValueWithoutReinstancing=true
 bDisplaysLayout=false
 bNativeCodeGenerationTool=false
-bSkeletonInheritSkeletonClasses=false
 bOptimizeExecutionFlowStack=true
 bOptimizeAdjacentStates=true
 
diff --git a/Engine/Source/Editor/KismetCompiler/Private/KismetCompiler.cpp b/Engine/Source/Editor/KismetCompiler/Private/KismetCompiler.cpp
--- a/Engine/Source/Editor/KismetCompiler/Private/KismetCompiler.cpp
+++ b/Engine/Source/Editor/KismetCompiler/Private/KismetCompiler.cpp
@@ -156,7 +156,7 @@ void FKismetCompilerContext::CleanAndSanitizeClass(UBlueprintGeneratedClass* Cla
 	
 	UClass* ParentClass = Blueprint->ParentClass;
 
-	if(UBlueprintGeneratedClass::CompileSkeletonClassesInheritSkeletonClasses() && FKismetEditorUtilities::IsClassABlueprintSkeleton(ClassToClean))
+	if(FKismetEditorUtilities::IsClassABlueprintSkeleton(ClassToClean))
 	{
 		if(UBlueprint* BlueprintParent = Cast<UBlueprint>(Blueprint->ParentClass->ClassGeneratedBy))
 		{
@@ -3271,14 +3271,7 @@ void FKismetCompilerContext::Compile()
 	NewClass->ClassGeneratedBy = Blueprint;
 
 	UClass* ParentClass = nullptr;
-	if(UBlueprintGeneratedClass::CompileSkeletonClassesInheritSkeletonClasses())
-	{
-		ParentClass = NewClass->ClassWithin;
-	}
-	else
-	{
-		ParentClass = Blueprint->ParentClass;
-	}
+	ParentClass = NewClass->ClassWithin;
 	NewClass->SetSuperStruct(ParentClass);
 	NewClass->ClassFlags |= (ParentClass->ClassFlags & CLASS_Inherit);
 	NewClass->ClassCastFlags |= ParentClass->ClassCastFlags;
diff --git a/Engine/Source/Editor/UnrealEd/Private/Kismet2/BlueprintEditorUtils.cpp b/Engine/Source/Editor/UnrealEd/Private/Kismet2/BlueprintEditorUtils.cpp
--- a/Engine/Source/Editor/UnrealEd/Private/Kismet2/BlueprintEditorUtils.cpp
+++ b/Engine/Source/Editor/UnrealEd/Private/Kismet2/BlueprintEditorUtils.cpp
@@ -3721,17 +3721,7 @@ void FBlueprintEditorUtils::ChangeMemberVariableType(UBlueprint* Blueprint, cons
 					Variable.VarType = NewPinType;
 
 					UClass* ParentClass = nullptr;
-					// When compiling so that skeleton classes inherit skeleton classes, marking the Blueprint structurally modified is enough
-					// when the children skeleton classes are rebuilt, they will receive the modified variables.
-					if(UBlueprintGeneratedClass::CompileSkeletonClassesInheritSkeletonClasses())
-					{
-						FBlueprintEditorUtils::MarkBlueprintAsStructurallyModified(Blueprint);
-					}
-					else
-					{
-						// Compile the Blueprint even if no loaded BPs are using the variable, this ensures that BPs refresh correctly when they are loaded
-						FKismetEditorUtilities::CompileBlueprint(Blueprint, false, true);
-					}
+					FBlueprintEditorUtils::MarkBlueprintAsStructurallyModified(Blueprint);
 
 					if(bBreakingVariableConnections)
 					{
diff --git a/Engine/Source/Editor/UnrealEd/Private/Kismet2/KismetReinstanceUtilities.cpp b/Engine/Source/Editor/UnrealEd/Private/Kismet2/KismetReinstanceUtilities.cpp
--- a/Engine/Source/Editor/UnrealEd/Private/Kismet2/KismetReinstanceUtilities.cpp
+++ b/Engine/Source/Editor/UnrealEd/Private/Kismet2/KismetReinstanceUtilities.cpp
@@ -838,14 +838,12 @@ void FBlueprintCompileReinstancer::ReparentChild(UBlueprint* ChildBP)
 	UClass* SkeletonClass = ChildBP->SkeletonGeneratedClass;
 	UClass* GeneratedClass = ChildBP->GeneratedClass;
 
-	const bool bReinstancingSkeleton = (UBlueprintGeneratedClass::CompileSkeletonClassesInheritSkeletonClasses() && bIsReinstancingSkeleton) || !UBlueprintGeneratedClass::CompileSkeletonClassesInheritSkeletonClasses();
-	if(  bReinstancingSkeleton && SkeletonClass )
+	if( bIsReinstancingSkeleton && SkeletonClass )
 	{
 		ReparentChild(SkeletonClass);
 	}
 
-	const bool bReinstancingGenerated = (UBlueprintGeneratedClass::CompileSkeletonClassesInheritSkeletonClasses() && !bIsReinstancingSkeleton) || !UBlueprintGeneratedClass::CompileSkeletonClassesInheritSkeletonClasses();
-	if( bReinstancingGenerated && GeneratedClass )
+	if( !bIsReinstancingSkeleton && GeneratedClass )
 	{
 		ReparentChild(GeneratedClass);
 	}
diff --git a/Engine/Source/Runtime/Engine/Classes/Engine/BlueprintGeneratedClass.h b/Engine/Source/Runtime/Engine/Classes/Engine/BlueprintGeneratedClass.h
--- a/Engine/Source/Runtime/Engine/Classes/Engine/BlueprintGeneratedClass.h
+++ b/Engine/Source/Runtime/Engine/Classes/Engine/BlueprintGeneratedClass.h
@@ -479,8 +479,6 @@ public:
 	static FName GetUberGraphFrameName();
 	static bool UsePersistentUberGraphFrame();
 
-	static bool CompileSkeletonClassesInheritSkeletonClasses();
-
 #if WITH_EDITORONLY_DATA
 	FBlueprintDebugData DebugData;
 
diff --git a/Engine/Source/Runtime/Engine/Private/BlueprintGeneratedClass.cpp b/Engine/Source/Runtime/Engine/Private/BlueprintGeneratedClass.cpp
--- a/Engine/Source/Runtime/Engine/Private/BlueprintGeneratedClass.cpp
+++ b/Engine/Source/Runtime/Engine/Private/BlueprintGeneratedClass.cpp
@@ -584,10 +584,4 @@ bool UBlueprintGeneratedClass::UsePersistentUberGraphFrame()
 #else
 	return false;
 #endif
-}
-
-bool UBlueprintGeneratedClass::CompileSkeletonClassesInheritSkeletonClasses()
-{
-	static const FBoolConfigValueHelper SkeletonClassesInheritSkeletonClasses(TEXT("Kismet"), TEXT("bSkeletonInheritSkeletonClasses"), GEngineIni);
-	return SkeletonClassesInheritSkeletonClasses;
 }
\ No newline at end of file
