[INTEGRATE] Change 2420565 by [REDACTED_PERSON][REDACTED_HANDLE] on 2015/01/27 12:10:55

	#ue4 Fixing up issue with Blueprint Components not showing up in the context menu.

diff --git a/Engine/Source/Editor/KismetCompiler/Private/KismetCompiler.cpp b/Engine/Source/Editor/KismetCompiler/Private/KismetCompiler.cpp
--- a/Engine/Source/Editor/KismetCompiler/Private/KismetCompiler.cpp
+++ b/Engine/Source/Editor/KismetCompiler/Private/KismetCompiler.cpp
@@ -1704,7 +1704,9 @@ void FKismetCompilerContext::FinishCompilingClass(UClass* Class)
 		// Blueprinted Components are always Blueprint Spawnable
 		if (ParentClass->IsChildOf(UActorComponent::StaticClass()))
 		{
-			Class->SetMetaData(TEXT("BlueprintSpawnableComponent"), TEXT("true"));
+			static const FName NAME_ClassGroupNames(TEXT("ClassGroupNames"));
+			Class->SetMetaData(FBlueprintMetadata::MD_BlueprintSpawnableComponent, TEXT("true"));
+			Class->SetMetaData(NAME_ClassGroupNames, *NSLOCTEXT("BlueprintableComponents", "CategoryName", "Custom").ToString());
 		}
 
 
