[INTEGRATE] Change 2422346 by [REDACTED_PERSON][REDACTED_HANDLE] on 2015/01/28 11:47:29

	Allow placement of a ReceiveTick event in Blueprinted ActorComponents
	bActive no longer required to keep ticking in components
	bAutoActivate no longer default in newly created C++ component classes

diff --git a/Engine/Source/Editor/DetailCustomizations/Private/ActorComponentDetails.cpp b/Engine/Source/Editor/DetailCustomizations/Private/ActorComponentDetails.cpp
new file mode 100644
--- /dev/null
+++ b/Engine/Source/Editor/DetailCustomizations/Private/ActorComponentDetails.cpp
@@ -0,0 +1,26 @@
+﻿// Copyright 1998-2015 Epic Games, Inc. All Rights Reserved.
+
+#include "DetailCustomizationsPrivatePCH.h"
+#include "ActorComponentDetails.h"
+
+TSharedRef<IDetailCustomization> FActorComponentDetails::MakeInstance()
+{
+	return MakeShareable( new FActorComponentDetails );
+}
+
+void FActorComponentDetails::CustomizeDetails( IDetailLayoutBuilder& DetailBuilder )
+{
+	TSharedPtr<IPropertyHandle> PrimaryTickProperty = DetailBuilder.GetProperty(GET_MEMBER_NAME_CHECKED(UActorComponent, PrimaryComponentTick));
+
+	// Defaults only show tick properties
+	if (DetailBuilder.GetDetailsView().HasClassDefaultObject())
+	{
+		IDetailCategoryBuilder& TickCategory = DetailBuilder.EditCategory("Tick");
+
+		TickCategory.AddProperty(PrimaryTickProperty->GetChildHandle(GET_MEMBER_NAME_CHECKED(FTickFunction, bStartWithTickEnabled)));
+		TickCategory.AddProperty(PrimaryTickProperty->GetChildHandle(GET_MEMBER_NAME_CHECKED(FTickFunction, bTickEvenWhenPaused)), EPropertyLocation::Advanced);
+		TickCategory.AddProperty(PrimaryTickProperty->GetChildHandle(GET_MEMBER_NAME_CHECKED(FTickFunction, bAllowTickOnDedicatedServer)), EPropertyLocation::Advanced);
+	}
+
+	PrimaryTickProperty->MarkHiddenByCustomization();
+}
diff --git a/Engine/Source/Editor/DetailCustomizations/Private/ActorComponentDetails.h b/Engine/Source/Editor/DetailCustomizations/Private/ActorComponentDetails.h
new file mode 100644
--- /dev/null
+++ b/Engine/Source/Editor/DetailCustomizations/Private/ActorComponentDetails.h
@@ -0,0 +1,13 @@
+// Copyright 1998-2015 Epic Games, Inc. All Rights Reserved.
+
+#pragma once
+
+class FActorComponentDetails : public IDetailCustomization
+{
+public:
+	/** Makes a new instance of this detail layout class for a specific detail view requesting it */
+	static TSharedRef<IDetailCustomization> MakeInstance();
+
+	/** IDetailCustomization interface */
+	virtual void CustomizeDetails( IDetailLayoutBuilder& DetailBuilder ) override;
+};
diff --git a/Engine/Source/Editor/DetailCustomizations/Private/DetailCustomizations.cpp b/Engine/Source/Editor/DetailCustomizations/Private/DetailCustomizations.cpp
--- a/Engine/Source/Editor/DetailCustomizations/Private/DetailCustomizations.cpp
+++ b/Engine/Source/Editor/DetailCustomizations/Private/DetailCustomizations.cpp
@@ -200,6 +200,7 @@ void FDetailCustomizationsModule::RegisterObjectCustomizations()
 	// In the order below, Actor will get a chance to display details first, followed by USceneComponent.
 
 	RegisterCustomClassLayout("Actor", FOnGetDetailCustomizationInstance::CreateStatic(&FActorDetails::MakeInstance));
+	RegisterCustomClassLayout("ActorComponent", FOnGetDetailCustomizationInstance::CreateStatic(&FActorComponentDetails::MakeInstance));
 	RegisterCustomClassLayout("SceneComponent", FOnGetDetailCustomizationInstance::CreateStatic(&FSceneComponentDetails::MakeInstance));
 	RegisterCustomClassLayout("PrimitiveComponent", FOnGetDetailCustomizationInstance::CreateStatic(&FPrimitiveComponentDetails::MakeInstance));
 	RegisterCustomClassLayout("StaticMeshComponent", FOnGetDetailCustomizationInstance::CreateStatic(&FStaticMeshComponentDetails::MakeInstance));
diff --git a/Engine/Source/Editor/Kismet/Private/BlueprintDetailsCustomization.cpp b/Engine/Source/Editor/Kismet/Private/BlueprintDetailsCustomization.cpp
--- a/Engine/Source/Editor/Kismet/Private/BlueprintDetailsCustomization.cpp
+++ b/Engine/Source/Editor/Kismet/Private/BlueprintDetailsCustomization.cpp
@@ -4511,6 +4511,10 @@ void FBlueprintComponentDetails::CustomizeDetails(IDetailLayoutBuilder& DetailLa
 				]
 			];
 	}
+
+	// Don't show tick properties for components in the blueprint details
+	TSharedPtr<IPropertyHandle> PrimaryTickProperty = DetailLayout.GetProperty(GET_MEMBER_NAME_CHECKED(UActorComponent, PrimaryComponentTick));
+	PrimaryTickProperty->MarkHiddenByCustomization();
 }
 
 UClass* FBlueprintComponentDetails::FindCommonBaseClassFromSelected() const
diff --git a/Engine/Source/Editor/KismetCompiler/Private/KismetCompiler.cpp b/Engine/Source/Editor/KismetCompiler/Private/KismetCompiler.cpp
--- a/Engine/Source/Editor/KismetCompiler/Private/KismetCompiler.cpp
+++ b/Engine/Source/Editor/KismetCompiler/Private/KismetCompiler.cpp
@@ -3531,7 +3531,7 @@ void FKismetCompilerContext::Compile()
 		}
 
 		CopyTermDefaultsToDefaultObject(NewCDO);
-		SetCanEverTickForActor();
+		SetCanEverTick();
 		FKismetCompilerUtilities::ValidateEnumProperties(NewCDO, MessageLog);
 	}
 
@@ -3704,48 +3704,56 @@ const UK2Node_FunctionEntry* FKismetCompilerContext::FindLocalEntryPoint(const U
 	return NULL;
 }
 
-void FKismetCompilerContext::SetCanEverTickForActor()
+void FKismetCompilerContext::SetCanEverTick()
 {
-	AActor* const CDActor = NewClass ? Cast<AActor>(NewClass->GetDefaultObject()) : NULL;
-	if(NULL == CDActor)
+	FTickFunction* TickFunction = nullptr;
+	FTickFunction* ParentTickFunction = nullptr;
+	
+	if (AActor* CDActor = Cast<AActor>(NewClass->GetDefaultObject()))
 	{
-		return;
+		TickFunction = &CDActor->PrimaryActorTick;
+		ParentTickFunction = &NewClass->GetSuperClass()->GetDefaultObject<AActor>()->PrimaryActorTick;
+	}
+	else if (UActorComponent* CDComponent = Cast<UActorComponent>(NewClass->GetDefaultObject()))
+	{
+		TickFunction = &CDComponent->PrimaryComponentTick;
+		ParentTickFunction = &NewClass->GetSuperClass()->GetDefaultObject<UActorComponent>()->PrimaryComponentTick;
 	}
 
-	const bool bOldFlag = CDActor->PrimaryActorTick.bCanEverTick;
-	// RESET FLAG 
+	if (TickFunction == nullptr)
 	{
-		UClass* ParentClass = NewClass->GetSuperClass();
-		const AActor* ParentCDO = ParentClass ? Cast<AActor>(ParentClass->GetDefaultObject()) : NULL;
-		check(NULL != ParentCDO);
-		// Clear to handle case, when an event (that forced a flag) was removed, or class was re-parented
-		CDActor->PrimaryActorTick.bCanEverTick = ParentCDO->PrimaryActorTick.bCanEverTick;
+		return;
 	}
+
+	const bool bOldFlag = TickFunction->bCanEverTick;
+	// RESET FLAG 
+	TickFunction->bCanEverTick = ParentTickFunction->bCanEverTick;
 	
 	// RECEIVE TICK
 	static FName ReceiveTickName(GET_FUNCTION_NAME_CHECKED(AActor, ReceiveTick));
+	static FName ComponentReceiveTickName(GET_FUNCTION_NAME_CHECKED(UActorComponent, ReceiveTick)); // Only doing this to ensure that both classes have the correct, same name
 	const UFunction* ReciveTickEvent = FKismetCompilerUtilities::FindOverriddenImplementableEvent(ReceiveTickName, NewClass);
 	if (ReciveTickEvent)
 	{
 		static const FName ChildCanTickName = TEXT("ChildCanTick");
 		const UClass* FirstNativeClass = FBlueprintEditorUtils::FindFirstNativeClass(NewClass);
-		const bool bOverrideFlags = (AActor::StaticClass() == FirstNativeClass) || (FirstNativeClass && FirstNativeClass->HasMetaData(ChildCanTickName));
-		if(bOverrideFlags)
+		const bool bOverrideFlags = (AActor::StaticClass() == FirstNativeClass) || (UActorComponent::StaticClass() == FirstNativeClass) || (FirstNativeClass && FirstNativeClass->HasMetaData(ChildCanTickName));
+		if (bOverrideFlags)
 		{
-			CDActor->PrimaryActorTick.bCanEverTick = true;
+			TickFunction->bCanEverTick = true;
 		}
-		else if(!CDActor->PrimaryActorTick.bCanEverTick)
+		else if (!TickFunction->bCanEverTick)
 		{
 			const FString ReceivTickEventWarning = FString::Printf( 
-				*LOCTEXT("ReceiveTick_CanNeverTick", "Blueprint %s has the ReceiveTick @@ event, but it can never tick.  Please consider using a Timer instead of a tick, using Actor as the parent class, or you can enable Tick using code in one of the following ways: set ChildCanTick in the metadata on the parent class, or set bCanEverTick to true.").ToString(), *NewClass->GetName());
+				*LOCTEXT("ReceiveTick_CanNeverTick", "Blueprint %s has the ReceiveTick @@ event, but it can never tick.  Please consider using a Timer instead of a tick or you can enable Tick using code in one of the following ways: set ChildCanTick in the metadata on the parent class, or set bCanEverTick to true.").ToString(), *NewClass->GetName());
 			MessageLog.Warning( *ReceivTickEventWarning, FindLocalEntryPoint(ReciveTickEvent) );
 		}
 	}
 
-	if(CDActor->PrimaryActorTick.bCanEverTick != bOldFlag)
+	if(TickFunction->bCanEverTick != bOldFlag)
 	{
-		UE_LOG(LogK2Compiler, Verbose, TEXT("Overridden flags for Actor class '%s': CanEverTick %s "), *NewClass->GetName(),
-			CDActor->PrimaryActorTick.bCanEverTick ? *(GTrue.ToString()) : *(GFalse.ToString()) );
+		UE_LOG(LogK2Compiler, Verbose, TEXT("Overridden flags for class '%s': CanEverTick %s "), *NewClass->GetName(),
+			TickFunction->bCanEverTick ? *(GTrue.ToString()) : *(GFalse.ToString()) );
 	}
 }
 
diff --git a/Engine/Source/Editor/KismetCompiler/Public/KismetCompiler.h b/Engine/Source/Editor/KismetCompiler/Public/KismetCompiler.h
--- a/Engine/Source/Editor/KismetCompiler/Public/KismetCompiler.h
+++ b/Engine/Source/Editor/KismetCompiler/Public/KismetCompiler.h
@@ -236,11 +236,10 @@ protected:
 	virtual void CopyTermDefaultsToDefaultObject(UObject* DefaultObject);
 
 	/** 
-	 * Function works only is first native superclass is AActor.
+	 * Function works only is subclass of AActor or UActorComponent.
 	 * If ReceiveTick event is defined, force CanEverTick.
-	 * If ReceiveAnyDamage or ReceiveRadialDamage or ReceivePointDamage event is defined, force CanBeDamaged.
 	 */
-	void SetCanEverTickForActor();
+	void SetCanEverTick();
 
 	/** Scan FunctionList and return Entry point, for matching one  */
 	const UK2Node_FunctionEntry* FindLocalEntryPoint(const UFunction* Function) const;
diff --git a/Engine/Source/Runtime/Engine/Classes/Components/ActorComponent.h b/Engine/Source/Runtime/Engine/Classes/Components/ActorComponent.h
--- a/Engine/Source/Runtime/Engine/Classes/Components/ActorComponent.h
+++ b/Engine/Source/Runtime/Engine/Classes/Components/ActorComponent.h
@@ -39,7 +39,7 @@ public:
 	virtual void GetLifetimeReplicatedProps(TArray<FLifetimeProperty>& OutLifetimeProps) const override;
 
 	/** Main tick function for the Actor */
-	UPROPERTY()
+	UPROPERTY(EditDefaultsOnly, Category="Tick")
 	struct FActorComponentTickFunction PrimaryComponentTick;
 
 	/** Array of tags that can be used for grouping and categorizing. Can also be accessed from scripting. */
@@ -334,6 +334,7 @@ public:
 	/** 
 	 * Returns whether this component has tick enabled or not
 	 */
+	UFUNCTION(BlueprintCallable, Category="Utilities")
 	bool IsComponentTickEnabled() const;
 
 	/**
@@ -539,6 +540,10 @@ public:
 	UFUNCTION(BlueprintCallable, Category="Utilities", meta=(Keywords = "dependency"))
 	virtual void RemoveTickPrerequisiteComponent(UActorComponent* PrerequisiteComponent);
 
+	/** Event called every frame */
+	UFUNCTION(BlueprintImplementableEvent, meta=(FriendlyName = "Tick"))
+	virtual void ReceiveTick(float DeltaSeconds);
+	
 	/** 
 	 *  Called by owner actor on position shifting
 	 *  Component should update all relevant data structures to reflect new actor location
diff --git a/Engine/Source/Runtime/Engine/Private/ActorComponent.cpp b/Engine/Source/Runtime/Engine/Private/ActorComponent.cpp
--- a/Engine/Source/Runtime/Engine/Private/ActorComponent.cpp
+++ b/Engine/Source/Runtime/Engine/Private/ActorComponent.cpp
@@ -579,10 +579,7 @@ void UActorComponent::TickComponent(float DeltaTime, enum ELevelTick TickType, F
 {
 	check(bRegistered);
 
-	if (!bIsActive)
-	{
-		SetComponentTickEnabled(false);
-	}
+	ReceiveTick(DeltaTime);
 }
 
 void UActorComponent::RegisterComponentWithWorld(UWorld* InWorld)
