[INTEGRATE] Change 2422712 by [REDACTED_PERSON][REDACTED_HANDLE] on 2015/01/28 15:24:05

	Allow InitializeComponent and UninitializeComponent to be bound for Blueprinted ActorComponents

diff --git a/Engine/Source/Editor/KismetCompiler/Private/KismetCompiler.cpp b/Engine/Source/Editor/KismetCompiler/Private/KismetCompiler.cpp
--- a/Engine/Source/Editor/KismetCompiler/Private/KismetCompiler.cpp
+++ b/Engine/Source/Editor/KismetCompiler/Private/KismetCompiler.cpp
@@ -3532,6 +3532,7 @@ void FKismetCompilerContext::Compile()
 
 		CopyTermDefaultsToDefaultObject(NewCDO);
 		SetCanEverTick();
+		SetWantsInitialize();
 		FKismetCompilerUtilities::ValidateEnumProperties(NewCDO, MessageLog);
 	}
 
@@ -3704,7 +3705,7 @@ const UK2Node_FunctionEntry* FKismetCompilerContext::FindLocalEntryPoint(const U
 	return NULL;
 }
 
-void FKismetCompilerContext::SetCanEverTick()
+void FKismetCompilerContext::SetCanEverTick() const
 {
 	FTickFunction* TickFunction = nullptr;
 	FTickFunction* ParentTickFunction = nullptr;
@@ -3730,33 +3731,68 @@ void FKismetCompilerContext::SetCanEverTick()
 	TickFunction->bCanEverTick = ParentTickFunction->bCanEverTick;
 	
 	// RECEIVE TICK
-	static FName ReceiveTickName(GET_FUNCTION_NAME_CHECKED(AActor, ReceiveTick));
-	static FName ComponentReceiveTickName(GET_FUNCTION_NAME_CHECKED(UActorComponent, ReceiveTick)); // Only doing this to ensure that both classes have the correct, same name
-	const UFunction* ReciveTickEvent = FKismetCompilerUtilities::FindOverriddenImplementableEvent(ReceiveTickName, NewClass);
-	if (ReciveTickEvent)
+	if (!TickFunction->bCanEverTick)
 	{
-		static const FName ChildCanTickName = TEXT("ChildCanTick");
-		const UClass* FirstNativeClass = FBlueprintEditorUtils::FindFirstNativeClass(NewClass);
-		const bool bOverrideFlags = (AActor::StaticClass() == FirstNativeClass) || (UActorComponent::StaticClass() == FirstNativeClass) || (FirstNativeClass && FirstNativeClass->HasMetaData(ChildCanTickName));
-		if (bOverrideFlags)
+		static FName ReceiveTickName(GET_FUNCTION_NAME_CHECKED(AActor, ReceiveTick));
+		static FName ComponentReceiveTickName(GET_FUNCTION_NAME_CHECKED(UActorComponent, ReceiveTick)); // Only doing this to ensure that both classes have the correct, same name
+		const UFunction* ReciveTickEvent = FKismetCompilerUtilities::FindOverriddenImplementableEvent(ReceiveTickName, NewClass);
+		if (ReciveTickEvent)
 		{
-			TickFunction->bCanEverTick = true;
-		}
-		else if (!TickFunction->bCanEverTick)
-		{
-			const FString ReceivTickEventWarning = FString::Printf( 
-				*LOCTEXT("ReceiveTick_CanNeverTick", "Blueprint %s has the ReceiveTick @@ event, but it can never tick.  Please consider using a Timer instead of a tick or you can enable Tick using code in one of the following ways: set ChildCanTick in the metadata on the parent class, or set bCanEverTick to true.").ToString(), *NewClass->GetName());
-			MessageLog.Warning( *ReceivTickEventWarning, FindLocalEntryPoint(ReciveTickEvent) );
+			static const FName ChildCanTickName = TEXT("ChildCanTick");
+			const UClass* FirstNativeClass = FBlueprintEditorUtils::FindFirstNativeClass(NewClass);
+			const bool bOverrideFlags = (AActor::StaticClass() == FirstNativeClass) || (UActorComponent::StaticClass() == FirstNativeClass) || (FirstNativeClass && FirstNativeClass->HasMetaData(ChildCanTickName));
+			if (bOverrideFlags)
+			{
+				TickFunction->bCanEverTick = true;
+			}
+			else if (!TickFunction->bCanEverTick)
+			{
+				const FString ReceivTickEventWarning = FString::Printf( 
+					*LOCTEXT("ReceiveTick_CanNeverTick", "Blueprint %s has the ReceiveTick @@ event, but it can never tick.  Please consider using a Timer instead of a tick or you can enable Tick using code in one of the following ways: set ChildCanTick in the metadata on the parent class, or set bCanEverTick to true.").ToString(), *NewClass->GetName());
+				MessageLog.Warning( *ReceivTickEventWarning, FindLocalEntryPoint(ReciveTickEvent) );
+			}
 		}
 	}
 
 	if(TickFunction->bCanEverTick != bOldFlag)
 	{
-		UE_LOG(LogK2Compiler, Verbose, TEXT("Overridden flags for class '%s': CanEverTick %s "), *NewClass->GetName(),
+		UE_LOG(LogK2Compiler, Verbose, TEXT("Overridden flag for class '%s': CanEverTick %s "), *NewClass->GetName(),
 			TickFunction->bCanEverTick ? *(GTrue.ToString()) : *(GFalse.ToString()) );
 	}
 }
 
+void FKismetCompilerContext::SetWantsInitialize() const
+{
+	UActorComponent* CDComponent = Cast<UActorComponent>(NewClass->GetDefaultObject());
+
+	if (CDComponent == nullptr)
+	{
+		return;
+	}
+
+	const bool bOldFlag = CDComponent->bWantsInitializeComponent;
+
+	CDComponent->bWantsInitializeComponent = NewClass->GetSuperClass()->GetDefaultObject<UActorComponent>()->bWantsInitializeComponent;
+
+	if (!CDComponent->bWantsInitializeComponent)
+	{
+		static FName ReceiveInitializeComponentName(GET_FUNCTION_NAME_CHECKED(UActorComponent, ReceiveInitializeComponent));
+		static FName ReceiveUninitializeComponentName(GET_FUNCTION_NAME_CHECKED(UActorComponent, ReceiveUninitializeComponent));
+		const UFunction* ReciveInitializeComponentEvent = FKismetCompilerUtilities::FindOverriddenImplementableEvent(ReceiveInitializeComponentName, NewClass);
+		const UFunction* ReciveUninitializeComponentEvent = FKismetCompilerUtilities::FindOverriddenImplementableEvent(ReceiveUninitializeComponentName, NewClass);
+		if (ReciveInitializeComponentEvent || ReciveUninitializeComponentEvent)
+		{
+			CDComponent->bWantsInitializeComponent = true;
+		}
+	}
+
+	if(CDComponent->bWantsInitializeComponent != bOldFlag)
+	{
+		UE_LOG(LogK2Compiler, Verbose, TEXT("Overridden flag for class '%s': bWantsInitializeComponent %s "), *NewClass->GetName(),
+			CDComponent->bWantsInitializeComponent ? *(GTrue.ToString()) : *(GFalse.ToString()) );
+	}
+}
+
 bool FKismetCompilerContext::UsePersistentUberGraphFrame() const
 {
 	return UBlueprintGeneratedClass::UsePersistentUberGraphFrame() && !CompileOptions.DoesRequireCppCodeGeneration();
diff --git a/Engine/Source/Editor/KismetCompiler/Public/KismetCompiler.h b/Engine/Source/Editor/KismetCompiler/Public/KismetCompiler.h
--- a/Engine/Source/Editor/KismetCompiler/Public/KismetCompiler.h
+++ b/Engine/Source/Editor/KismetCompiler/Public/KismetCompiler.h
@@ -236,10 +236,16 @@ protected:
 	virtual void CopyTermDefaultsToDefaultObject(UObject* DefaultObject);
 
 	/** 
-	 * Function works only is subclass of AActor or UActorComponent.
+	 * Function works only if subclass of AActor or UActorComponent.
 	 * If ReceiveTick event is defined, force CanEverTick.
 	 */
-	void SetCanEverTick();
+	void SetCanEverTick() const;
+
+	/** 
+	 * Function works only if subclass UActorComponent.
+	 * If InitializeComponent or UninitializeComponent event is defined, force bWantsInitializeComponent.
+	 */
+	void SetWantsInitialize() const;
 
 	/** Scan FunctionList and return Entry point, for matching one  */
 	const UK2Node_FunctionEntry* FindLocalEntryPoint(const UFunction* Function) const;
diff --git a/Engine/Source/Runtime/Engine/Classes/Components/ActorComponent.h b/Engine/Source/Runtime/Engine/Classes/Components/ActorComponent.h
--- a/Engine/Source/Runtime/Engine/Classes/Components/ActorComponent.h
+++ b/Engine/Source/Runtime/Engine/Classes/Components/ActorComponent.h
@@ -112,6 +112,7 @@ public:
 	uint32 bIsActive:1;
 
 	/** If TRUE, we call the virtual InitializeComponent */
+	UPROPERTY()
 	uint32 bWantsInitializeComponent:1;
 
 	/** Indicates that OnCreatedComponent has been called, but OnDestroyedComponent has not yet */
@@ -287,12 +288,20 @@ public:
 	 */
 	virtual void InitializeComponent();
 
+	/** Event when the component is initialized, either via creation or its Actor's BeginPlay. */
+	UFUNCTION(BlueprintImplementableEvent, meta=(Keywords = "Begin", FriendlyName = "Initialize Component"))
+	virtual void ReceiveInitializeComponent();
+
 	/**
 	 * Ends gameplay for this component.
 	 * Called from AActor::EndPlay only if bHasBeenInitialized is true
 	 */
 	virtual void UninitializeComponent();
 
+	/** Event when the component is uninitialized, generally via descruction or its Actor's EndPlay. */
+	UFUNCTION(BlueprintImplementableEvent, meta=(Keywords = "End Delete", FriendlyName = "Uninitialize Component"))
+	virtual void ReceiveUninitializeComponent();
+	
 	/**
 	 * When called, will call the virtual call chain to register all of the tick functions
 	 * Do not override this function or make it virtual
diff --git a/Engine/Source/Runtime/Engine/Private/ActorComponent.cpp b/Engine/Source/Runtime/Engine/Private/ActorComponent.cpp
--- a/Engine/Source/Runtime/Engine/Private/ActorComponent.cpp
+++ b/Engine/Source/Runtime/Engine/Private/ActorComponent.cpp
@@ -475,12 +475,22 @@ void UActorComponent::InitializeComponent()
 {
 	check(bRegistered);
 	check(!bHasBeenInitialized);
+
+	ReceiveInitializeComponent();
+
 	bHasBeenInitialized = true;
 }
 
 void UActorComponent::UninitializeComponent()
 {
 	check(bHasBeenInitialized);
+
+	// If we're already pending kill blueprints don't get to be notified
+	if (!HasAnyFlags(RF_BeginDestroyed))
+	{
+		ReceiveUninitializeComponent();
+	}
+
 	bHasBeenInitialized = false;
 }
 
