Trim down the amount of UObject constructors.

diff --git a/Engine/Source/Editor/Kismet/Private/STimelineEditor.cpp b/Engine/Source/Editor/Kismet/Private/STimelineEditor.cpp
--- a/Engine/Source/Editor/Kismet/Private/STimelineEditor.cpp
+++ b/Engine/Source/Editor/Kismet/Private/STimelineEditor.cpp
@@ -975,7 +975,7 @@ FReply STimelineEditor::CreateNewTrack(FTimelineEdTrack::ETrackType Type)
 			{
 				FTTEventTrack NewTrack;
 				NewTrack.TrackName = TrackName;
-				NewTrack.CurveKeys = NewNamedObject<UCurveFloat>(OwnerClass, NAME_None, RF_Public); // Needs to be marked public so that it can be referenced from timeline instances in the level
+				NewTrack.CurveKeys = NewObject<UCurveFloat>(OwnerClass, NAME_None, RF_Public); // Needs to be marked public so that it can be referenced from timeline instances in the level
 				NewTrack.CurveKeys->bIsEventCurve = true;
 				TimelineObj->EventTracks.Add(NewTrack);
 			}
@@ -987,7 +987,7 @@ FReply STimelineEditor::CreateNewTrack(FTimelineEdTrack::ETrackType Type)
 				NewTrack.CurveFloat = FindObject<UCurveFloat>(ANY_PACKAGE, *TrackName.ToString() );
 				if (NewTrack.CurveFloat == NULL)
 				{
-					NewTrack.CurveFloat = NewNamedObject<UCurveFloat>(OwnerClass, NAME_None, RF_Public);
+					NewTrack.CurveFloat = NewObject<UCurveFloat>(OwnerClass, NAME_None, RF_Public);
 				}
 				TimelineObj->FloatTracks.Add(NewTrack);
 			}
@@ -995,14 +995,14 @@ FReply STimelineEditor::CreateNewTrack(FTimelineEdTrack::ETrackType Type)
 			{
 				FTTVectorTrack NewTrack;
 				NewTrack.TrackName = TrackName;
-				NewTrack.CurveVector = NewNamedObject<UCurveVector>(OwnerClass, NAME_None, RF_Public);
+				NewTrack.CurveVector = NewObject<UCurveVector>(OwnerClass, NAME_None, RF_Public);
 				TimelineObj->VectorTracks.Add(NewTrack);
 			}
 			else if(Type == FTimelineEdTrack::TT_LinearColorInterp)
 			{
 				FTTLinearColorTrack NewTrack;
 				NewTrack.TrackName = TrackName;
-				NewTrack.CurveLinearColor = NewNamedObject<UCurveLinearColor>(OwnerClass, NAME_None, RF_Public);
+				NewTrack.CurveLinearColor = NewObject<UCurveLinearColor>(OwnerClass, NAME_None, RF_Public);
 				TimelineObj->LinearColorTracks.Add(NewTrack);
 			}
 
@@ -1048,19 +1048,19 @@ UCurveBase* STimelineEditor::CreateNewCurve( FTimelineEdTrack::ETrackType Type )
 	UCurveBase* NewCurve = NULL;
 	if(Type == FTimelineEdTrack::TT_Event)
 	{
-		NewCurve = NewNamedObject<UCurveFloat>(OwnerClass, NAME_None, RF_Public);
+		NewCurve = NewObject<UCurveFloat>(OwnerClass, NAME_None, RF_Public);
 	}
 	else if(Type == FTimelineEdTrack::TT_FloatInterp)
 	{
-		NewCurve = NewNamedObject<UCurveFloat>(OwnerClass, NAME_None, RF_Public);
+		NewCurve = NewObject<UCurveFloat>(OwnerClass, NAME_None, RF_Public);
 	}
 	else if(Type == FTimelineEdTrack::TT_VectorInterp)
 	{
-		NewCurve = NewNamedObject<UCurveVector>(OwnerClass, NAME_None, RF_Public);
+		NewCurve = NewObject<UCurveVector>(OwnerClass, NAME_None, RF_Public);
 	}
 	else if(Type == FTimelineEdTrack::TT_LinearColorInterp)
 	{
-		NewCurve = NewNamedObject<UCurveLinearColor>(OwnerClass, NAME_None, RF_Public);
+		NewCurve = NewObject<UCurveLinearColor>(OwnerClass, NAME_None, RF_Public);
 	}
 
 
diff --git a/Engine/Source/Editor/KismetCompiler/Private/AnimBlueprintCompiler.cpp b/Engine/Source/Editor/KismetCompiler/Private/AnimBlueprintCompiler.cpp
--- a/Engine/Source/Editor/KismetCompiler/Private/AnimBlueprintCompiler.cpp
+++ b/Engine/Source/Editor/KismetCompiler/Private/AnimBlueprintCompiler.cpp
@@ -1117,7 +1117,7 @@ void FAnimBlueprintCompiler::SpawnNewClass(const FString& NewClassName)
 
 	if (NewAnimBlueprintClass == NULL)
 	{
-		NewAnimBlueprintClass = ConstructObject<UAnimBlueprintGeneratedClass>(UAnimBlueprintGeneratedClass::StaticClass(), Blueprint->GetOutermost(), FName(*NewClassName), RF_Public|RF_Transactional);
+		NewAnimBlueprintClass = NewObject<UAnimBlueprintGeneratedClass>(Blueprint->GetOutermost(), FName(*NewClassName), RF_Public | RF_Transactional);
 	}
 	else
 	{
diff --git a/Engine/Source/Editor/KismetCompiler/Private/KismetCompiler.cpp b/Engine/Source/Editor/KismetCompiler/Private/KismetCompiler.cpp
--- a/Engine/Source/Editor/KismetCompiler/Private/KismetCompiler.cpp
+++ b/Engine/Source/Editor/KismetCompiler/Private/KismetCompiler.cpp
@@ -113,7 +113,7 @@ void FKismetCompilerContext::SpawnNewClass(const FString& NewClassName)
 	if (NewClass == NULL)
 	{
 		// If the class hasn't been found, then spawn a new one
-		NewClass = ConstructObject<UBlueprintGeneratedClass>(UBlueprintGeneratedClass::StaticClass(), Blueprint->GetOutermost(), FName(*NewClassName), RF_Public|RF_Transactional);
+		NewClass = NewObject<UBlueprintGeneratedClass>(Blueprint->GetOutermost(), FName(*NewClassName), RF_Public | RF_Transactional);
 	}
 	else
 	{
@@ -148,7 +148,7 @@ void FKismetCompilerContext::CleanAndSanitizeClass(UBlueprintGeneratedClass* Cla
 	const bool bRecompilingOnLoad = Blueprint->bIsRegeneratingOnLoad;
 	FString TransientClassString = FString::Printf(TEXT("TRASHCLASS_%s"), *Blueprint->GetName());
 	FName TransientClassName = MakeUniqueObjectName(GetTransientPackage(), UBlueprintGeneratedClass::StaticClass(), FName(*TransientClassString));
-	UClass* TransientClass = ConstructObject<UBlueprintGeneratedClass>(UBlueprintGeneratedClass::StaticClass(), GetTransientPackage(), TransientClassName, RF_Public|RF_Transient);
+	UClass* TransientClass = NewObject<UBlueprintGeneratedClass>(GetTransientPackage(), TransientClassName, RF_Public | RF_Transient);
 	
 	UClass* ParentClass = Blueprint->ParentClass;
 
@@ -1169,7 +1169,7 @@ void FKismetCompilerContext::PrecompileFunction(FKismetFunctionContext& Context)
 			CreatedFunctionNames.Add(NewFunctionNameString);
 		}
 
-		Context.Function = NewNamedObject<UFunction>(NewClass, NewFunctionName, RF_Public);
+		Context.Function = NewObject<UFunction>(NewClass, NewFunctionName, RF_Public);
 
 #if USE_TRANSIENT_SKELETON
 		// Propagate down transient settings from the class
@@ -1822,7 +1822,7 @@ void FKismetCompilerContext::BuildDynamicBindingObjects(UBlueprintGeneratedClass
 					UDynamicBlueprintBinding* DynamicBindingObject = Class->GetDynamicBindingObject(DynamicBindingClass);
 					if (DynamicBindingObject == NULL)
 					{
-						DynamicBindingObject = ConstructObject<UDynamicBlueprintBinding>(DynamicBindingClass, Class);
+						DynamicBindingObject = NewObject<UDynamicBlueprintBinding>(Class, DynamicBindingClass);
 						Class->DynamicBindingObjects.Add(DynamicBindingObject);
 					}
 					Node->RegisterDynamicBinding(DynamicBindingObject);
@@ -2327,7 +2327,7 @@ void FKismetCompilerContext::CreateFunctionStubForEvent(UK2Node_Event* SrcEventN
 
 	// Create the stub graph and add it to the list of functions to compile
 
-	UEdGraph* ChildStubGraph = NewNamedObject<UEdGraph>(OwnerOfTemporaries, EventNodeName);
+	UEdGraph* ChildStubGraph = NewObject<UEdGraph>(OwnerOfTemporaries, EventNodeName);
 	Blueprint->EventGraphs.Add(ChildStubGraph);
 	ChildStubGraph->Schema = UEdGraphSchema_K2::StaticClass();
 	ChildStubGraph->SetFlags(RF_Transient);
@@ -2625,7 +2625,7 @@ void FKismetCompilerContext::CreateAndProcessUbergraph()
 {
 	BP_SCOPED_COMPILER_EVENT_STAT(EKismetCompilerStats_ProcessUbergraph);
 
-	ConsolidatedEventGraph = NewNamedObject<UEdGraph>(Blueprint, GetUbergraphCallName());
+	ConsolidatedEventGraph = NewObject<UEdGraph>(Blueprint, GetUbergraphCallName());
 	ConsolidatedEventGraph->Schema = UEdGraphSchema_K2::StaticClass();
 	ConsolidatedEventGraph->SetFlags(RF_Transient);
 
diff --git a/Engine/Source/Editor/KismetCompiler/Private/KismetCompilerMisc.cpp b/Engine/Source/Editor/KismetCompiler/Private/KismetCompilerMisc.cpp
--- a/Engine/Source/Editor/KismetCompiler/Private/KismetCompilerMisc.cpp
+++ b/Engine/Source/Editor/KismetCompiler/Private/KismetCompilerMisc.cpp
@@ -676,7 +676,7 @@ UProperty* FKismetCompilerUtilities::CreatePropertyOnScope(UStruct* Scope, const
 	UArrayProperty* NewArrayProperty = NULL;
 	if( bIsArrayProperty )
 	{
-		NewArrayProperty = NewNamedObject<UArrayProperty>(Scope, ValidatedPropertyName, ObjectFlags);
+		NewArrayProperty = NewObject<UArrayProperty>(Scope, ValidatedPropertyName, ObjectFlags);
 		PropertyScope = NewArrayProperty;
 	}
 	else
@@ -699,7 +699,7 @@ UProperty* FKismetCompilerUtilities::CreatePropertyOnScope(UStruct* Scope, const
 		{
 			if (SubType->HasAnyClassFlags(CLASS_Interface))
 			{
-				UInterfaceProperty* NewPropertyObj = NewNamedObject<UInterfaceProperty>(PropertyScope, ValidatedPropertyName, ObjectFlags);
+				UInterfaceProperty* NewPropertyObj = NewObject<UInterfaceProperty>(PropertyScope, ValidatedPropertyName, ObjectFlags);
 				// we want to use this setter function instead of setting the 
 				// InterfaceClass member directly, because it properly handles  
 				// placeholder classes (classes that are stubbed in during load)
@@ -712,11 +712,11 @@ UProperty* FKismetCompilerUtilities::CreatePropertyOnScope(UStruct* Scope, const
 
 				if( Type.bIsWeakPointer )
 				{
-					NewPropertyObj = NewNamedObject<UWeakObjectProperty>(PropertyScope, ValidatedPropertyName, ObjectFlags);
+					NewPropertyObj = NewObject<UWeakObjectProperty>(PropertyScope, ValidatedPropertyName, ObjectFlags);
 				}
 				else
 				{
-					NewPropertyObj = NewNamedObject<UObjectProperty>(PropertyScope, ValidatedPropertyName, ObjectFlags);
+					NewPropertyObj = NewObject<UObjectProperty>(PropertyScope, ValidatedPropertyName, ObjectFlags);
 				}
 				// we want to use this setter function instead of setting the 
 				// PropertyClass member directly, because it properly handles  
@@ -734,7 +734,7 @@ UProperty* FKismetCompilerUtilities::CreatePropertyOnScope(UStruct* Scope, const
 			FString StructureError;
 			if (FStructureEditorUtils::EStructureError::Ok == FStructureEditorUtils::IsStructureValid(SubType, NULL, &StructureError))
 			{
-				UStructProperty* NewPropertyStruct = NewNamedObject<UStructProperty>(PropertyScope, ValidatedPropertyName, ObjectFlags);
+				UStructProperty* NewPropertyStruct = NewObject<UStructProperty>(PropertyScope, ValidatedPropertyName, ObjectFlags);
 				NewPropertyStruct->Struct = SubType;
 				NewProperty = NewPropertyStruct;
 			}
@@ -768,7 +768,7 @@ UProperty* FKismetCompilerUtilities::CreatePropertyOnScope(UStruct* Scope, const
 
 		if (SubType != NULL)
 		{
-			UClassProperty* NewPropertyClass = NewNamedObject<UClassProperty>(PropertyScope, ValidatedPropertyName, ObjectFlags);
+			UClassProperty* NewPropertyClass = NewObject<UClassProperty>(PropertyScope, ValidatedPropertyName, ObjectFlags);
 			// we want to use this setter function instead of setting the 
 			// MetaClass member directly, because it properly handles  
 			// placeholder classes (classes that are stubbed in during load)
@@ -781,7 +781,7 @@ UProperty* FKismetCompilerUtilities::CreatePropertyOnScope(UStruct* Scope, const
 	{
 		if (UFunction* SignatureFunction = FMemberReference::ResolveSimpleMemberReference<UFunction>(Type.PinSubCategoryMemberReference))
 		{
-			UDelegateProperty* NewPropertyDelegate = NewNamedObject<UDelegateProperty>(PropertyScope, ValidatedPropertyName, ObjectFlags);
+			UDelegateProperty* NewPropertyDelegate = NewObject<UDelegateProperty>(PropertyScope, ValidatedPropertyName, ObjectFlags);
 			NewPropertyDelegate->SignatureFunction = SignatureFunction;
 			NewProperty = NewPropertyDelegate;
 		}
@@ -789,47 +789,47 @@ UProperty* FKismetCompilerUtilities::CreatePropertyOnScope(UStruct* Scope, const
 	else if (Type.PinCategory == Schema->PC_MCDelegate)
 	{
 		UFunction* const SignatureFunction = FMemberReference::ResolveSimpleMemberReference<UFunction>(Type.PinSubCategoryMemberReference);
-		UMulticastDelegateProperty* NewPropertyDelegate = NewNamedObject<UMulticastDelegateProperty>(PropertyScope, ValidatedPropertyName, ObjectFlags);
+		UMulticastDelegateProperty* NewPropertyDelegate = NewObject<UMulticastDelegateProperty>(PropertyScope, ValidatedPropertyName, ObjectFlags);
 		NewPropertyDelegate->SignatureFunction = SignatureFunction;
 		NewProperty = NewPropertyDelegate;
 	}
 	else if (Type.PinCategory == Schema->PC_Int)
 	{
-		NewProperty = NewNamedObject<UIntProperty>(PropertyScope, ValidatedPropertyName, ObjectFlags);
+		NewProperty = NewObject<UIntProperty>(PropertyScope, ValidatedPropertyName, ObjectFlags);
 	}
 	else if (Type.PinCategory == Schema->PC_Float)
 	{
-		NewProperty = NewNamedObject<UFloatProperty>(PropertyScope, ValidatedPropertyName, ObjectFlags);
+		NewProperty = NewObject<UFloatProperty>(PropertyScope, ValidatedPropertyName, ObjectFlags);
 	}
 	else if (Type.PinCategory == Schema->PC_Boolean)
 	{
-		UBoolProperty* BoolProperty = NewNamedObject<UBoolProperty>(PropertyScope, ValidatedPropertyName, ObjectFlags);
+		UBoolProperty* BoolProperty = NewObject<UBoolProperty>(PropertyScope, ValidatedPropertyName, ObjectFlags);
 		BoolProperty->SetBoolSize(sizeof(bool), true);
 		NewProperty = BoolProperty;
 	}
 	else if (Type.PinCategory == Schema->PC_String)
 	{
-		NewProperty = NewNamedObject<UStrProperty>(PropertyScope, ValidatedPropertyName, ObjectFlags);
+		NewProperty = NewObject<UStrProperty>(PropertyScope, ValidatedPropertyName, ObjectFlags);
 	}
 	else if (Type.PinCategory == Schema->PC_Text)
 	{
-		NewProperty = NewNamedObject<UTextProperty>(PropertyScope, ValidatedPropertyName, ObjectFlags);
+		NewProperty = NewObject<UTextProperty>(PropertyScope, ValidatedPropertyName, ObjectFlags);
 	}
 	else if (Type.PinCategory == Schema->PC_Byte)
 	{
-		UByteProperty* ByteProp = NewNamedObject<UByteProperty>(PropertyScope, ValidatedPropertyName, ObjectFlags);
+		UByteProperty* ByteProp = NewObject<UByteProperty>(PropertyScope, ValidatedPropertyName, ObjectFlags);
 		ByteProp->Enum = Cast<UEnum>(Type.PinSubCategoryObject.Get());
 
 		NewProperty = ByteProp;
 	}
 	else if (Type.PinCategory == Schema->PC_Name)
 	{
-		NewProperty = NewNamedObject<UNameProperty>(PropertyScope, ValidatedPropertyName, ObjectFlags);
+		NewProperty = NewObject<UNameProperty>(PropertyScope, ValidatedPropertyName, ObjectFlags);
 	}
 	else
 	{
 		// Failed to resolve the type-subtype, create a generic property to survive VM bytecode emission
-		NewProperty = NewNamedObject<UIntProperty>(PropertyScope, ValidatedPropertyName, ObjectFlags);
+		NewProperty = NewObject<UIntProperty>(PropertyScope, ValidatedPropertyName, ObjectFlags);
 	}
 
 	if (bIsArrayProperty)
diff --git a/Engine/Source/Editor/KismetCompiler/Private/UserDefinedStructureCompilerUtils.cpp b/Engine/Source/Editor/KismetCompiler/Private/UserDefinedStructureCompilerUtils.cpp
--- a/Engine/Source/Editor/KismetCompiler/Private/UserDefinedStructureCompilerUtils.cpp
+++ b/Engine/Source/Editor/KismetCompiler/Private/UserDefinedStructureCompilerUtils.cpp
@@ -75,7 +75,7 @@ struct FUserDefinedStructureCompilerInner
 
 		const FString TransientString = FString::Printf(TEXT("TRASHSTRUCT_%s"), *StructToClean->GetName());
 		const FName TransientName = MakeUniqueObjectName(GetTransientPackage(), UUserDefinedStruct::StaticClass(), FName(*TransientString));
-		UUserDefinedStruct* TransientStruct = NewNamedObject<UUserDefinedStruct>(GetTransientPackage(), TransientName, RF_Public|RF_Transient);
+		UUserDefinedStruct* TransientStruct = NewObject<UUserDefinedStruct>(GetTransientPackage(), TransientName, RF_Public | RF_Transient);
 
 		TArray<UObject*> SubObjects;
 		GetObjectsWithOuter(StructToClean, SubObjects, true);
diff --git a/Engine/Source/Editor/LandscapeEditor/Private/LandscapeEdMode.cpp b/Engine/Source/Editor/LandscapeEditor/Private/LandscapeEdMode.cpp
--- a/Engine/Source/Editor/LandscapeEditor/Private/LandscapeEdMode.cpp
+++ b/Engine/Source/Editor/LandscapeEditor/Private/LandscapeEdMode.cpp
@@ -206,7 +206,7 @@ FEdModeLandscape::FEdModeLandscape()
 	CurrentToolTarget.TargetType = ELandscapeToolTargetType::Heightmap;
 	CurrentToolTarget.LayerInfo = nullptr;
 
-	UISettings = ConstructObject<ULandscapeEditorObject>(ULandscapeEditorObject::StaticClass(), GetTransientPackage(), NAME_None, RF_Transactional);
+	UISettings = NewObject<ULandscapeEditorObject>(GetTransientPackage(), TEXT("UISettings"), RF_Transactional);
 	UISettings->SetParent(this);
 }
 
@@ -331,7 +331,7 @@ void FEdModeLandscape::Enter()
 	if (CurrentGizmoActor.IsValid() && !CurrentGizmoActor->GizmoTexture)
 	{
 		// Init Gizmo Texture...
-		CurrentGizmoActor->GizmoTexture = ConstructObject<UTexture2D>(UTexture2D::StaticClass(), GetTransientPackage(), NAME_None, RF_Transient);
+		CurrentGizmoActor->GizmoTexture = NewObject<UTexture2D>(GetTransientPackage(), NAME_None, RF_Transient);
 		if (CurrentGizmoActor->GizmoTexture)
 		{
 			CurrentGizmoActor->GizmoTexture->Source.Init(
diff --git a/Engine/Source/Editor/LandscapeEditor/Private/LandscapeEdModeComponentTools.cpp b/Engine/Source/Editor/LandscapeEditor/Private/LandscapeEdModeComponentTools.cpp
--- a/Engine/Source/Editor/LandscapeEditor/Private/LandscapeEdModeComponentTools.cpp
+++ b/Engine/Source/Editor/LandscapeEditor/Private/LandscapeEdModeComponentTools.cpp
@@ -820,7 +820,7 @@ public:
 					{
 						// Add New component...
 						FIntPoint ComponentBase = FIntPoint(ComponentIndexX, ComponentIndexY)*Landscape->ComponentSizeQuads;
-						ULandscapeComponent* LandscapeComponent = ConstructObject<ULandscapeComponent>(ULandscapeComponent::StaticClass(), Landscape, NAME_None, RF_Transactional);
+						ULandscapeComponent* LandscapeComponent = NewObject<ULandscapeComponent>(Landscape, NAME_None, RF_Transactional);
 						Landscape->LandscapeComponents.Add(LandscapeComponent);
 						NewComponents.Add(LandscapeComponent);
 						LandscapeComponent->Init(
diff --git a/Engine/Source/Editor/LandscapeEditor/Private/LandscapeEdModePaintTools.cpp b/Engine/Source/Editor/LandscapeEditor/Private/LandscapeEdModePaintTools.cpp
--- a/Engine/Source/Editor/LandscapeEditor/Private/LandscapeEdModePaintTools.cpp
+++ b/Engine/Source/Editor/LandscapeEditor/Private/LandscapeEdModePaintTools.cpp
@@ -665,7 +665,7 @@ public:
 		FLandscapeToolPaintBase<ToolTarget, FLandscapeToolStrokeFlatten<ToolTarget>>::EnterTool();
 
 		ALandscapeProxy* LandscapeProxy = this->EdMode->CurrentToolTarget.LandscapeInfo->GetLandscapeProxy();
-		MeshComponent = ConstructObject<UStaticMeshComponent>(UStaticMeshComponent::StaticClass(), LandscapeProxy, NAME_None, RF_Transient);
+		MeshComponent = NewObject<UStaticMeshComponent>(LandscapeProxy, NAME_None, RF_Transient);
 		MeshComponent->StaticMesh = PlaneMesh;
 		MeshComponent->AttachTo(LandscapeProxy->GetRootComponent());
 		MeshComponent->RegisterComponent();
diff --git a/Engine/Source/Editor/LandscapeEditor/Private/LandscapeEdModeSplineTools.cpp b/Engine/Source/Editor/LandscapeEditor/Private/LandscapeEdModeSplineTools.cpp
--- a/Engine/Source/Editor/LandscapeEditor/Private/LandscapeEdModeSplineTools.cpp
+++ b/Engine/Source/Editor/LandscapeEditor/Private/LandscapeEdModeSplineTools.cpp
@@ -65,7 +65,7 @@ public:
 	void CreateSplineComponent(ALandscapeProxy* Landscape, FVector Scale3D)
 	{
 		Landscape->Modify();
-		Landscape->SplineComponent = ConstructObject<ULandscapeSplinesComponent>(ULandscapeSplinesComponent::StaticClass(), Landscape, NAME_None, RF_Transactional);
+		Landscape->SplineComponent = NewObject<ULandscapeSplinesComponent>(Landscape, NAME_None, RF_Transactional);
 		Landscape->SplineComponent->RelativeScale3D = Scale3D;
 		Landscape->SplineComponent->AttachTo(Landscape->GetRootComponent());
 		Landscape->SplineComponent->ShowSplineEditorMesh(true);
@@ -255,7 +255,7 @@ public:
 		Start->Modify();
 		End->Modify();
 
-		ULandscapeSplineSegment* NewSegment = ConstructObject<ULandscapeSplineSegment>(ULandscapeSplineSegment::StaticClass(), SplinesComponent, NAME_None, RF_Transactional);
+		ULandscapeSplineSegment* NewSegment = NewObject<ULandscapeSplineSegment>(SplinesComponent, NAME_None, RF_Transactional);
 		SplinesComponent->Segments.Add(NewSegment);
 
 		NewSegment->Connections[0].ControlPoint = Start;
@@ -333,7 +333,7 @@ public:
 
 		Landscape->SplineComponent->Modify();
 
-		ULandscapeSplineControlPoint* NewControlPoint = ConstructObject<ULandscapeSplineControlPoint>(ULandscapeSplineControlPoint::StaticClass(), Landscape->SplineComponent, NAME_None, RF_Transactional);
+		ULandscapeSplineControlPoint* NewControlPoint = NewObject<ULandscapeSplineControlPoint>(Landscape->SplineComponent, NAME_None, RF_Transactional);
 		Landscape->SplineComponent->ControlPoints.Add(NewControlPoint);
 
 		NewControlPoint->Location = LocalLocation;
@@ -512,7 +512,7 @@ public:
 		FVector Tangent;
 		Segment->FindNearest(LocalLocation, t, Location, Tangent);
 
-		ULandscapeSplineControlPoint* NewControlPoint = ConstructObject<ULandscapeSplineControlPoint>(ULandscapeSplineControlPoint::StaticClass(), SplinesComponent, NAME_None, RF_Transactional);
+		ULandscapeSplineControlPoint* NewControlPoint = NewObject<ULandscapeSplineControlPoint>(SplinesComponent, NAME_None, RF_Transactional);
 		SplinesComponent->ControlPoints.Add(NewControlPoint);
 
 		NewControlPoint->Location = Location;
@@ -522,7 +522,7 @@ public:
 		NewControlPoint->SideFalloff = FMath::Lerp(Segment->Connections[0].ControlPoint->SideFalloff, Segment->Connections[1].ControlPoint->SideFalloff, t);
 		NewControlPoint->EndFalloff = FMath::Lerp(Segment->Connections[0].ControlPoint->EndFalloff, Segment->Connections[1].ControlPoint->EndFalloff, t);
 
-		ULandscapeSplineSegment* NewSegment = ConstructObject<ULandscapeSplineSegment>(ULandscapeSplineSegment::StaticClass(), SplinesComponent, NAME_None, RF_Transactional);
+		ULandscapeSplineSegment* NewSegment = NewObject<ULandscapeSplineSegment>(SplinesComponent, NAME_None, RF_Transactional);
 		SplinesComponent->Segments.Add(NewSegment);
 
 		NewSegment->Connections[0].ControlPoint = NewControlPoint;
diff --git a/Engine/Source/Editor/LandscapeEditor/Private/LandscapeEditorDetailCustomization_NewLandscape.cpp b/Engine/Source/Editor/LandscapeEditor/Private/LandscapeEditorDetailCustomization_NewLandscape.cpp
--- a/Engine/Source/Editor/LandscapeEditor/Private/LandscapeEditorDetailCustomization_NewLandscape.cpp
+++ b/Engine/Source/Editor/LandscapeEditor/Private/LandscapeEditorDetailCustomization_NewLandscape.cpp
@@ -1579,7 +1579,7 @@ void FLandscapeEditorStructCustomization_FLandscapeImportLayer::OnImportLayerCre
 			LayerObjectName = FName(*NewLayerDlg->GetAssetName().ToString());
 
 			UPackage* Package = CreatePackage(NULL, *PackageName);
-			ULandscapeLayerInfoObject* LayerInfo = ConstructObject<ULandscapeLayerInfoObject>(ULandscapeLayerInfoObject::StaticClass(), Package, LayerObjectName, RF_Public | RF_Standalone | RF_Transactional);
+			ULandscapeLayerInfoObject* LayerInfo = NewObject<ULandscapeLayerInfoObject>(Package, LayerObjectName, RF_Public | RF_Standalone | RF_Transactional);
 			LayerInfo->LayerName = LayerName;
 			LayerInfo->bNoWeightBlend = bNoWeightBlend;
 
diff --git a/Engine/Source/Editor/LandscapeEditor/Private/LandscapeEditorDetailCustomization_TargetLayers.cpp b/Engine/Source/Editor/LandscapeEditor/Private/LandscapeEditorDetailCustomization_TargetLayers.cpp
--- a/Engine/Source/Editor/LandscapeEditor/Private/LandscapeEditorDetailCustomization_TargetLayers.cpp
+++ b/Engine/Source/Editor/LandscapeEditor/Private/LandscapeEditorDetailCustomization_TargetLayers.cpp
@@ -737,7 +737,7 @@ void FLandscapeEditorCustomNodeBuilder_TargetLayers::OnTargetLayerCreateClicked(
 			LayerObjectName = FName(*NewLayerDlg->GetAssetName().ToString());
 
 			UPackage* Package = CreatePackage(NULL, *PackageName);
-			ULandscapeLayerInfoObject* LayerInfo = ConstructObject<ULandscapeLayerInfoObject>(ULandscapeLayerInfoObject::StaticClass(), Package, LayerObjectName, RF_Public | RF_Standalone | RF_Transactional);
+			ULandscapeLayerInfoObject* LayerInfo = NewObject<ULandscapeLayerInfoObject>(Package, LayerObjectName, RF_Public | RF_Standalone | RF_Transactional);
 			LayerInfo->LayerName = LayerName;
 			LayerInfo->bNoWeightBlend = bNoWeightBlend;
 
diff --git a/Engine/Source/Editor/LandscapeEditor/Private/LandscapeEditorModule.cpp b/Engine/Source/Editor/LandscapeEditor/Private/LandscapeEditorModule.cpp
--- a/Engine/Source/Editor/LandscapeEditor/Private/LandscapeEditorModule.cpp
+++ b/Engine/Source/Editor/LandscapeEditor/Private/LandscapeEditorModule.cpp
@@ -57,11 +57,11 @@ public:
 		LevelEditorModule.GetMenuExtensibilityManager()->AddExtender(ViewportMenuExtender);
 
 		// Actor Factories
-		UActorFactoryLandscape* LandscapeActorFactory = ConstructObject<UActorFactoryLandscape>(UActorFactoryLandscape::StaticClass());
+		UActorFactoryLandscape* LandscapeActorFactory = NewObject<UActorFactoryLandscape>();
 		LandscapeActorFactory->NewActorClass = ALandscape::StaticClass();
 		GEditor->ActorFactories.Add(LandscapeActorFactory);
 
-		UActorFactoryLandscape* LandscapeProxyActorFactory = ConstructObject<UActorFactoryLandscape>(UActorFactoryLandscape::StaticClass());
+		UActorFactoryLandscape* LandscapeProxyActorFactory = NewObject<UActorFactoryLandscape>();
 		LandscapeProxyActorFactory->NewActorClass = ALandscapeProxy::StaticClass();
 		GEditor->ActorFactories.Add(LandscapeProxyActorFactory);
 	}
diff --git a/Engine/Source/Editor/LevelEditor/Private/LevelEditorActions.cpp b/Engine/Source/Editor/LevelEditor/Private/LevelEditorActions.cpp
--- a/Engine/Source/Editor/LevelEditor/Private/LevelEditorActions.cpp
+++ b/Engine/Source/Editor/LevelEditor/Private/LevelEditorActions.cpp
@@ -1970,7 +1970,7 @@ void FLevelEditorActionCallbacks::OpenLevelBlueprint( TWeakPtr< SLevelEditor > L
 void FLevelEditorActionCallbacks::CreateBlueprintClass()
 {
 	// Use the BlueprintFactory to allow the user to pick a parent class for the new Blueprint class
-	UBlueprintFactory* NewFactory = Cast<UBlueprintFactory>(ConstructObject<UFactory>( UBlueprintFactory::StaticClass() ));
+	UBlueprintFactory* NewFactory = Cast<UBlueprintFactory>(NewObject<UFactory>(GetTransientPackage(), UBlueprintFactory::StaticClass()));
 	FEditorDelegates::OnConfigureNewAssetProperties.Broadcast(NewFactory);
 	if ( NewFactory->ConfigureProperties() )
 	{
diff --git a/Engine/Source/Editor/LevelEditor/Private/SLevelViewport.cpp b/Engine/Source/Editor/LevelEditor/Private/SLevelViewport.cpp
--- a/Engine/Source/Editor/LevelEditor/Private/SLevelViewport.cpp
+++ b/Engine/Source/Editor/LevelEditor/Private/SLevelViewport.cpp
@@ -740,7 +740,7 @@ bool SLevelViewport::HandlePlaceDraggedObjects(const FGeometry& MyGeometry, cons
 		DroppedObjects.Empty();
 
 		// Create a container object to hold the export text and pass it into the actor placement code
-		UExportTextContainer* NewContainer = ConstructObject<UExportTextContainer>(UExportTextContainer::StaticClass());
+		UExportTextContainer* NewContainer = NewObject<UExportTextContainer>();
 		NewContainer->ExportText = DragDropOp->ActorExportText;
 		DroppedObjects.Add(NewContainer);
 	}
diff --git a/Engine/Source/Editor/LevelEditor/Private/SSurfaceProperties.cpp b/Engine/Source/Editor/LevelEditor/Private/SSurfaceProperties.cpp
--- a/Engine/Source/Editor/LevelEditor/Private/SSurfaceProperties.cpp
+++ b/Engine/Source/Editor/LevelEditor/Private/SSurfaceProperties.cpp
@@ -531,8 +531,7 @@ TSharedRef<SWidget> SSurfaceProperties::ConstructLighting()
 					}
 					if (FoundIndex == INDEX_NONE)
 					{
-						ULightmassPrimitiveSettingsObject* LightmassSettingsObject = 
-							ConstructObject<ULightmassPrimitiveSettingsObject>(ULightmassPrimitiveSettingsObject::StaticClass());
+						ULightmassPrimitiveSettingsObject* LightmassSettingsObject = NewObject<ULightmassPrimitiveSettingsObject>();
 						LightmassSettingsObject->LightmassSettings = TempSettings;
 						ObjArray.Add(LightmassSettingsObject);
 						SelectedLightmassSettingsObjects.Add(LightmassSettingsObject);
diff --git a/Engine/Source/Editor/MaterialEditor/Private/MaterialEditor.cpp b/Engine/Source/Editor/MaterialEditor/Private/MaterialEditor.cpp
--- a/Engine/Source/Editor/MaterialEditor/Private/MaterialEditor.cpp
+++ b/Engine/Source/Editor/MaterialEditor/Private/MaterialEditor.cpp
@@ -1051,7 +1051,7 @@ void FMaterialEditor::RefreshPreviewViewport()
 
 void FMaterialEditor::LoadEditorSettings()
 {
-	EditorOptions = ConstructObject<UMaterialEditorOptions>( UMaterialEditorOptions::StaticClass() );
+	EditorOptions = NewObject<UMaterialEditorOptions>();
 	
 	if (EditorOptions->bHideUnusedConnectors) {OnShowConnectors();}
 	if (bLivePreview != EditorOptions->bLivePreviewUpdate)
@@ -2805,7 +2805,7 @@ UMaterialExpression* FMaterialEditor::CreateNewMaterialExpression(UClass* NewExp
 			ExpressionOuter = MaterialFunction;
 		}
 
-		NewExpression = ConstructObject<UMaterialExpression>( NewExpressionClass, ExpressionOuter, NAME_None, RF_Transactional );
+		NewExpression = NewObject<UMaterialExpression>(ExpressionOuter, NewExpressionClass, NAME_None, RF_Transactional);
 		Material->Expressions.Add( NewExpression );
 		NewExpression->Material = Material;
 
@@ -2890,7 +2890,7 @@ UMaterialExpression* FMaterialEditor::CreateNewMaterialExpression(UClass* NewExp
 		if( RotateAboutAxisExpression )
 		{
 			// Create a default expression for the Position input
-			UMaterialExpressionWorldPosition* WorldPositionExpression = ConstructObject<UMaterialExpressionWorldPosition>( UMaterialExpressionWorldPosition::StaticClass(), ExpressionOuter, NAME_None, RF_Transactional );
+			UMaterialExpressionWorldPosition* WorldPositionExpression = NewObject<UMaterialExpressionWorldPosition>(ExpressionOuter, NAME_None, RF_Transactional);
 			Material->Expressions.Add( WorldPositionExpression );
 			WorldPositionExpression->Material = Material;
 			RotateAboutAxisExpression->Position.Expression = WorldPositionExpression;
@@ -2957,7 +2957,7 @@ UMaterialExpressionComment* FMaterialEditor::CreateNewMaterialExpressionComment(
 			ExpressionOuter = MaterialFunction;
 		}
 
-		NewComment = ConstructObject<UMaterialExpressionComment>( UMaterialExpressionComment::StaticClass(), ExpressionOuter, NAME_None, RF_Transactional );
+		NewComment = NewObject<UMaterialExpressionComment>(ExpressionOuter, NAME_None, RF_Transactional);
 
 		// Add to the list of comments associated with this material.
 		Material->EditorComments.Add( NewComment );
diff --git a/Engine/Source/Editor/MaterialEditor/Private/MaterialExpressionClasses.cpp b/Engine/Source/Editor/MaterialEditor/Private/MaterialExpressionClasses.cpp
--- a/Engine/Source/Editor/MaterialEditor/Private/MaterialExpressionClasses.cpp
+++ b/Engine/Source/Editor/MaterialEditor/Private/MaterialExpressionClasses.cpp
@@ -69,7 +69,7 @@ void MaterialExpressionClasses::InitMaterialExpressionClasses()
 {
 	if( !bInitialized )
 	{
-		UMaterialEditorOptions* TempEditorOptions = ConstructObject<UMaterialEditorOptions>( UMaterialEditorOptions::StaticClass() );
+		UMaterialEditorOptions* TempEditorOptions = NewObject<UMaterialEditorOptions>();
 		UClass* BaseType = UMaterialExpression::StaticClass();
 		if( BaseType )
 		{
diff --git a/Engine/Source/Editor/MaterialEditor/Private/MaterialInstanceEditor.cpp b/Engine/Source/Editor/MaterialEditor/Private/MaterialInstanceEditor.cpp
--- a/Engine/Source/Editor/MaterialEditor/Private/MaterialInstanceEditor.cpp
+++ b/Engine/Source/Editor/MaterialEditor/Private/MaterialInstanceEditor.cpp
@@ -156,7 +156,7 @@ void FMaterialInstanceEditor::InitMaterialInstanceEditor( const EToolkitMode::Ty
 	bShowMobileStats = false;
 
 	// Construct a temp holder for our instance parameters.
-	MaterialEditorInstance = ConstructObject<UMaterialEditorInstanceConstant>(UMaterialEditorInstanceConstant::StaticClass(), GetTransientPackage(), NAME_None, RF_Transactional);
+	MaterialEditorInstance = NewObject<UMaterialEditorInstanceConstant>(GetTransientPackage(), NAME_None, RF_Transactional);
 
 	bool bTempUseOldStyleMICEditorGroups = true;
 	GConfig->GetBool(TEXT("/Script/UnrealEd.EditorEngine"), TEXT("UseOldStyleMICEditorGroups"), bTempUseOldStyleMICEditorGroups, GEngineIni);	
diff --git a/Engine/Source/Editor/MaterialEditor/Private/SMaterialEditorViewport.cpp b/Engine/Source/Editor/MaterialEditor/Private/SMaterialEditorViewport.cpp
--- a/Engine/Source/Editor/MaterialEditor/Private/SMaterialEditorViewport.cpp
+++ b/Engine/Source/Editor/MaterialEditor/Private/SMaterialEditorViewport.cpp
@@ -194,10 +194,8 @@ void SMaterialEditorViewport::Construct(const FArguments& InArgs)
 	SEditorViewport::Construct( SEditorViewport::FArguments() );
 
 
-	PreviewMeshComponent = ConstructObject<UMaterialEditorMeshComponent>(
-		UMaterialEditorMeshComponent::StaticClass(), GetTransientPackage(), NAME_None, RF_Transient );
-	PreviewSkeletalMeshComponent = ConstructObject<USkeletalMeshComponent>(
-		USkeletalMeshComponent::StaticClass(), GetTransientPackage(), NAME_None, RF_Transient );
+	PreviewMeshComponent = NewObject<UMaterialEditorMeshComponent>(GetTransientPackage(), NAME_None, RF_Transient);
+	PreviewSkeletalMeshComponent = NewObject<USkeletalMeshComponent>(GetTransientPackage(), NAME_None, RF_Transient);
 
 	UMaterialInterface* Material = MaterialEditorPtr.Pin()->GetMaterialInterface();
 	if (Material)
diff --git a/Engine/Source/Editor/Matinee/Private/Matinee.cpp b/Engine/Source/Editor/Matinee/Private/Matinee.cpp
--- a/Engine/Source/Editor/Matinee/Private/Matinee.cpp
+++ b/Engine/Source/Editor/Matinee/Private/Matinee.cpp
@@ -400,7 +400,7 @@ void FMatinee::BuildCurveEditor()
 {
 	if(!IData->CurveEdSetup)
 	{
-		IData->CurveEdSetup = ConstructObject<UInterpCurveEdSetup>( UInterpCurveEdSetup::StaticClass(), IData, NAME_None );
+		IData->CurveEdSetup = NewObject<UInterpCurveEdSetup>(IData, NAME_None);
 	}
 
 	// Create graph editor to work on MatineeData's CurveEd setup.
@@ -490,7 +490,7 @@ void FMatinee::InitMatinee(const EToolkitMode::Type Mode, const TSharedPtr< clas
 	GetLockCameraPitchFromConfig();
 
 	// Create options object.
-	Opt = ConstructObject<UMatineeOptions>( UMatineeOptions::StaticClass(), GetTransientPackage(), NAME_None, RF_Transactional );
+	Opt = NewObject<UMatineeOptions>(GetTransientPackage(), NAME_None, RF_Transactional);
 	check(Opt);
 
 	// Swap out regular UTransactor for our special one
diff --git a/Engine/Source/Editor/Matinee/Private/MatineeMenus.cpp b/Engine/Source/Editor/Matinee/Private/MatineeMenus.cpp
--- a/Engine/Source/Editor/Matinee/Private/MatineeMenus.cpp
+++ b/Engine/Source/Editor/Matinee/Private/MatineeMenus.cpp
@@ -393,7 +393,7 @@ void FMatinee::NewGroupPopupTextCommitted(
 	GroupAttribs.Add(FAnalyticsEventAttribute(TEXT("ActionId"), FString::Printf(TEXT("%d"), static_cast<int32>(InActionId))));
 	if(bDirGroup)
 	{
-		UInterpGroup* NewDirector = ConstructObject<UInterpGroupDirector>( UInterpGroupDirector::StaticClass(), IData, NAME_None, RF_Transactional );
+		UInterpGroup* NewDirector = NewObject<UInterpGroupDirector>(IData, NAME_None, RF_Transactional);
 		NewGroups.Add(NewDirector);
 		GroupAttribs.Add(FAnalyticsEventAttribute(TEXT("Name"), NewDirector->GroupName.ToString()));
 	}
@@ -416,7 +416,7 @@ void FMatinee::NewGroupPopupTextCommitted(
 	}
 	else
 	{
-		UInterpGroup* NewGroup = ConstructObject<UInterpGroup>( UInterpGroup::StaticClass(), IData, NAME_None, RF_Transactional );
+		UInterpGroup* NewGroup = NewObject<UInterpGroup>(IData, NAME_None, RF_Transactional);
 		NewGroup->GroupName = NewGroupName;
 		NewGroups.Add(NewGroup);
 		GroupAttribs.Add(FAnalyticsEventAttribute(TEXT("Name"), NewGroup->GroupName.ToString()));
@@ -483,12 +483,12 @@ void FMatinee::NewGroupPopupTextCommitted(
 			// Create new InterpGroupInst
 			if(bDirGroup)
 			{
-				NewGroupInst = ConstructObject<UInterpGroupInstDirector>( UInterpGroupInstDirector::StaticClass(), MatineeActor, NAME_None, RF_Transactional );
+				NewGroupInst = NewObject<UInterpGroupInstDirector>(MatineeActor, NAME_None, RF_Transactional);
 				NewGroupInst->InitGroupInst(NewGroup, NULL);
 			}
 			else
 			{
-				NewGroupInst = ConstructObject<UInterpGroupInst>( UInterpGroupInst::StaticClass(), MatineeActor, NAME_None, RF_Transactional );
+				NewGroupInst = NewObject<UInterpGroupInst>(MatineeActor, NAME_None, RF_Transactional);
 				// Initialize group instance, saving ref to actor it works on.
 				NewGroupInst->InitGroupInst(NewGroup, GroupActor);
 			}
@@ -504,10 +504,10 @@ void FMatinee::NewGroupPopupTextCommitted(
 		// If a director group, create a director track for it now.
 		if(bDirGroup)
 		{
-			UInterpTrack* NewDirTrack = ConstructObject<UInterpTrackDirector>( UInterpTrackDirector::StaticClass(), NewGroup, NAME_None, RF_Transactional );
+			UInterpTrack* NewDirTrack = NewObject<UInterpTrackDirector>(NewGroup, NAME_None, RF_Transactional);
 			const int32 TrackIndex = NewGroup->InterpTracks.Add(NewDirTrack);
 
-			UInterpTrackInst* NewDirTrackInst = ConstructObject<UInterpTrackInstDirector>( UInterpTrackInstDirector::StaticClass(), NewGroupInst, NAME_None, RF_Transactional );
+			UInterpTrackInst* NewDirTrackInst = NewObject<UInterpTrackInstDirector>(NewGroupInst, NAME_None, RF_Transactional);
 			NewGroupInst->TrackInst.Add(NewDirTrackInst);
 
 			NewDirTrackInst->InitTrackInst(NewDirTrack);
@@ -551,10 +551,10 @@ void FMatinee::NewGroupPopupTextCommitted(
 			// For Lighting groups, add a Movement, Brightness, Light Color, and Radius Property track
 			if ( InActionId == FMatineeCommands::EGroupAction::NewLightingGroup )
 			{
-				UInterpTrack* NewMovTrack = ConstructObject<UInterpTrackMove>( UInterpTrackMove::StaticClass(), NewGroup, NAME_None, RF_Transactional );
+				UInterpTrack* NewMovTrack = NewObject<UInterpTrackMove>(NewGroup, NAME_None, RF_Transactional);
 				const int32 TrackIndex = NewGroup->InterpTracks.Add(NewMovTrack);
 
-				UInterpTrackInst* NewMovTrackInst = ConstructObject<UInterpTrackInstMove>( UInterpTrackInstMove::StaticClass(), NewGroupInst, NAME_None, RF_Transactional );
+				UInterpTrackInst* NewMovTrackInst = NewObject<UInterpTrackInstMove>(NewGroupInst, NAME_None, RF_Transactional);
 				NewGroupInst->TrackInst.Add(NewMovTrackInst);
 
 				NewMovTrackInst->InitTrackInst(NewMovTrack);
@@ -1434,7 +1434,7 @@ void FMatinee::AddActorToGroup(UInterpGroup* GroupToAdd, AActor* ActorToAdd)
 	}
 	else
 	{
-		NewGroupInst = ConstructObject<UInterpGroupInst>( UInterpGroupInst::StaticClass(), MatineeActor, NAME_None, RF_Transactional );
+		NewGroupInst = NewObject<UInterpGroupInst>(MatineeActor, NAME_None, RF_Transactional);
 		// Instantiate the Matinee group data structure.
 		MatineeActor->GroupInst.Add(NewGroupInst);
 		NewGroupInst->InitGroupInst(GroupToAdd, ActorToAdd);
@@ -1822,7 +1822,7 @@ void FMatinee::OnContextGroupCreateTabTextCommitted(const FText& InText, ETextCo
 		// Create a new tab.
 		if( HasAGroupSelected() )
 		{
-			UInterpFilter_Custom* Filter = ConstructObject<UInterpFilter_Custom>(UInterpFilter_Custom::StaticClass(), IData, NAME_None, RF_Transactional);
+			UInterpFilter_Custom* Filter = NewObject<UInterpFilter_Custom>(IData, NAME_None, RF_Transactional);
 
 			if(!InText.IsEmpty())
 			{
diff --git a/Engine/Source/Editor/Matinee/Private/MatineeTools.cpp b/Engine/Source/Editor/Matinee/Private/MatineeTools.cpp
--- a/Engine/Source/Editor/Matinee/Private/MatineeTools.cpp
+++ b/Engine/Source/Editor/Matinee/Private/MatineeTools.cpp
@@ -183,7 +183,7 @@ void  FMatinee::UpdateCameraRecording (void)
 						NewCam->GetCameraComponent()->FieldOfView = LevelVC->ViewFOV;
 
 						//make new group for the camera
-						UInterpGroup* NewGroup = ConstructObject<UInterpGroup>( UInterpGroup::StaticClass(), IData, NAME_None, RF_Transactional );
+						UInterpGroup* NewGroup = NewObject<UInterpGroup>(IData, NAME_None, RF_Transactional);
 						FString NewGroupName = NSLOCTEXT( "UnrealEd", "InterpEd_RecordMode_CameraGroupName", "CameraGroupCG" ).ToString();
 						NewGroup->GroupName = FName(*NewGroupName);
 						NewGroup->EnsureUniqueName();
@@ -192,7 +192,7 @@ void  FMatinee::UpdateCameraRecording (void)
 						IData->InterpGroups.Add(NewGroup);
 
 						//add group instance for camera
-						UInterpGroupInst* NewGroupInst = ConstructObject<UInterpGroupInst>( UInterpGroupInst::StaticClass(), MatineeActor, NAME_None, RF_Transactional );
+						UInterpGroupInst* NewGroupInst = NewObject<UInterpGroupInst>(MatineeActor, NAME_None, RF_Transactional);
 						// Initialise group instance, saving ref to actor it works on.
 						NewGroupInst->InitGroupInst(NewGroup, NewCam);
 						const int32 NewGroupInstIndex = MatineeActor->GroupInst.Add(NewGroupInst);
@@ -2780,7 +2780,7 @@ UInterpTrack* FMatinee::AddTrackToGroup( UInterpGroup* Group, UClass* TrackClass
 	}
 	else
 	{
-		NewTrack = ConstructObject<UInterpTrack>( TrackClass, Group, NAME_None, RF_Transactional );
+		NewTrack = NewObject<UInterpTrack>(Group, TrackClass, NAME_None, RF_Transactional);
 	}
 
 	check(NewTrack);
@@ -2812,7 +2812,7 @@ UInterpTrack* FMatinee::AddTrackToGroup( UInterpGroup* Group, UClass* TrackClass
 		{
 			GroupInst->Modify();
 
-			UInterpTrackInst* NewTrackInst = ConstructObject<UInterpTrackInst>( NewTrack->TrackInstClass, GroupInst, NAME_None, RF_Transactional );
+			UInterpTrackInst* NewTrackInst = NewObject<UInterpTrackInst>(GroupInst, NewTrack->TrackInstClass, NAME_None, RF_Transactional);
 
 			const int32 NewInstIndex = GroupInst->TrackInst.Add(NewTrackInst);
 			check(NewInstIndex == OutNewTrackIndex);
@@ -3305,11 +3305,11 @@ void FMatinee::DuplicateGroup(UInterpGroup* GroupToDuplicate)
 
 		if(bDirGroup)
 		{
-			NewGroupInst = ConstructObject<UInterpGroupInstDirector>( UInterpGroupInstDirector::StaticClass(), MatineeActor, NAME_None, RF_Transactional );
+			NewGroupInst = NewObject<UInterpGroupInstDirector>(MatineeActor, NAME_None, RF_Transactional);
 		}
 		else
 		{
-			NewGroupInst = ConstructObject<UInterpGroupInst>( UInterpGroupInst::StaticClass(), MatineeActor, NAME_None, RF_Transactional );
+			NewGroupInst = NewObject<UInterpGroupInst>(MatineeActor, NAME_None, RF_Transactional);
 		}
 
 		// Initialise group instance, saving ref to actor it works on.
@@ -3332,10 +3332,10 @@ void FMatinee::DuplicateGroup(UInterpGroup* GroupToDuplicate)
 			
 			if(NewDirTrack==NULL)
 			{
-				NewDirTrack = ConstructObject<UInterpTrackDirector>( UInterpTrackDirector::StaticClass(), NewGroup, NAME_None, RF_Transactional );
+				NewDirTrack = NewObject<UInterpTrackDirector>(NewGroup, NAME_None, RF_Transactional);
 				NewGroup->InterpTracks.Add(NewDirTrack);
 
-				UInterpTrackInst* NewDirTrackInst = ConstructObject<UInterpTrackInstDirector>( UInterpTrackInstDirector::StaticClass(), NewGroupInst, NAME_None, RF_Transactional );
+				UInterpTrackInst* NewDirTrackInst = NewObject<UInterpTrackInstDirector>(NewGroupInst, NAME_None, RF_Transactional);
 				NewGroupInst->TrackInst.Add(NewDirTrackInst);
 
 				NewDirTrackInst->InitTrackInst(NewDirTrack);
diff --git a/Engine/Source/Editor/MeshPaint/Private/MeshPaintEdMode.cpp b/Engine/Source/Editor/MeshPaint/Private/MeshPaintEdMode.cpp
--- a/Engine/Source/Editor/MeshPaint/Private/MeshPaintEdMode.cpp
+++ b/Engine/Source/Editor/MeshPaint/Private/MeshPaintEdMode.cpp
@@ -4889,7 +4889,7 @@ void FEdModeMeshPaint::DuplicateTextureMaterialCombo()
 				{
 					UClass* FactoryClass = UMaterialInstanceConstantFactoryNew::StaticClass();
 
-					UMaterialInstanceConstantFactoryNew* Factory = ConstructObject<UMaterialInstanceConstantFactoryNew>(UMaterialInstanceConstantFactoryNew::StaticClass());
+					UMaterialInstanceConstantFactoryNew* Factory = NewObject<UMaterialInstanceConstantFactoryNew>();
 					if ( Factory->ConfigureProperties() )
 					{
 						FString AssetName;
@@ -4953,7 +4953,7 @@ void FEdModeMeshPaint::CreateNewTexture()
 	{
 		UClass* FactoryClass = UTexture2DFactoryNew::StaticClass();
 			
-		UTexture2DFactoryNew* Factory = ConstructObject<UTexture2DFactoryNew>(UTexture2DFactoryNew::StaticClass());
+		UTexture2DFactoryNew* Factory = NewObject<UTexture2DFactoryNew>();
 		if ( Factory->ConfigureProperties() )
 		{
 			FString AssetName;
diff --git a/Engine/Source/Editor/NiagaraEditor/Private/SNiagaraEffectEditorViewport.cpp b/Engine/Source/Editor/NiagaraEditor/Private/SNiagaraEffectEditorViewport.cpp
--- a/Engine/Source/Editor/NiagaraEditor/Private/SNiagaraEffectEditorViewport.cpp
+++ b/Engine/Source/Editor/NiagaraEditor/Private/SNiagaraEffectEditorViewport.cpp
@@ -112,7 +112,7 @@ void SNiagaraEffectEditorViewport::Construct(const FArguments& InArgs)
 	
 	SEditorViewport::Construct( SEditorViewport::FArguments() );
 
-	PreviewComponent = ConstructObject<UNiagaraComponent>(UNiagaraComponent::StaticClass(), GetTransientPackage(), NAME_None, RF_Transient);
+	PreviewComponent = NewObject<UNiagaraComponent>(GetTransientPackage(), NAME_None, RF_Transient);
 }
 
 SNiagaraEffectEditorViewport::~SNiagaraEffectEditorViewport()
diff --git a/Engine/Source/Editor/Persona/Private/AnimationEditorViewportClient.cpp b/Engine/Source/Editor/Persona/Private/AnimationEditorViewportClient.cpp
--- a/Engine/Source/Editor/Persona/Private/AnimationEditorViewportClient.cpp
+++ b/Engine/Source/Editor/Persona/Private/AnimationEditorViewportClient.cpp
@@ -130,7 +130,7 @@ FAnimationViewportClient::FAnimationViewportClient(FPreviewScene& InPreviewScene
 	PreviewScene->DirectionalLight->RecreateRenderState_Concurrent();
 
 	// A background sky sphere
-	EditorSkyComp = ConstructObject<UStaticMeshComponent>(UStaticMeshComponent::StaticClass());
+	EditorSkyComp = NewObject<UStaticMeshComponent>(GetTransientPackage(), TEXT("EditorSkyComp"));
 	UStaticMesh * StaticMesh = LoadObject<UStaticMesh>(NULL, TEXT("/Engine/MapTemplates/Sky/SM_SkySphere.SM_SkySphere"), NULL, LOAD_None, NULL);
 	check (StaticMesh);
 	EditorSkyComp->SetStaticMesh( StaticMesh );
@@ -143,7 +143,7 @@ FAnimationViewportClient::FAnimationViewportClient(FPreviewScene& InPreviewScene
 
 	// now add height fog component
 
-	EditorHeightFogComponent = ConstructObject<UExponentialHeightFogComponent>(UExponentialHeightFogComponent::StaticClass());
+	EditorHeightFogComponent = NewObject<UExponentialHeightFogComponent>(GetTransientPackage(), TEXT("EditorHeightFogComponent"));
 	
 	EditorHeightFogComponent->FogDensity=0.00075f;
 	EditorHeightFogComponent->FogInscatteringColor=FLinearColor(3.f,4.f,6.f,0.f)*0.3f;
@@ -155,7 +155,7 @@ FAnimationViewportClient::FAnimationViewportClient(FPreviewScene& InPreviewScene
 	PreviewScene->AddComponent(EditorHeightFogComponent, FogTransform);
 
 	// add capture component for reflection
- 	USphereReflectionCaptureComponent* CaptureComponent = ConstructObject<USphereReflectionCaptureComponent>(USphereReflectionCaptureComponent::StaticClass());
+	USphereReflectionCaptureComponent* CaptureComponent = NewObject<USphereReflectionCaptureComponent>(GetTransientPackage(), TEXT("CaptureComponent"));
  
  	const FTransform CaptureTransform(FRotator(0,0,0), FVector(0.f,0.f,100.f), FVector(1.f));
  	PreviewScene->AddComponent(CaptureComponent, CaptureTransform);
@@ -165,7 +165,7 @@ FAnimationViewportClient::FAnimationViewportClient(FPreviewScene& InPreviewScene
 	// now add floor
  	UStaticMesh* FloorMesh = LoadObject<UStaticMesh>(NULL, TEXT("/Engine/EditorMeshes/PhAT_FloorBox.PhAT_FloorBox"), NULL, LOAD_None, NULL);
  	check(FloorMesh);
- 	EditorFloorComp = ConstructObject<UStaticMeshComponent>(UStaticMeshComponent::StaticClass());
+	EditorFloorComp = NewObject<UStaticMeshComponent>(GetTransientPackage(), TEXT("EditorFloorComp"));
  	EditorFloorComp->SetStaticMesh( FloorMesh );
  	PreviewScene->AddComponent(EditorFloorComp, FTransform::Identity);
  	EditorFloorComp->SetRelativeScale3D(FVector(3.f, 3.f, 1.f));
diff --git a/Engine/Source/Editor/Persona/Private/Persona.cpp b/Engine/Source/Editor/Persona/Private/Persona.cpp
--- a/Engine/Source/Editor/Persona/Private/Persona.cpp
+++ b/Engine/Source/Editor/Persona/Private/Persona.cpp
@@ -650,7 +650,7 @@ void FPersona::InitPersona(const EToolkitMode::Type Mode, const TSharedPtr< clas
 	check(PreviewComponent == NULL);
 
 	// Create the preview component
-	PreviewComponent = ConstructObject<UDebugSkelMeshComponent>(UDebugSkelMeshComponent::StaticClass(), GetTransientPackage());
+	PreviewComponent = NewObject<UDebugSkelMeshComponent>();
 
 	// note: we add to root here (rather than using RF_Standalone) as all standalone objects will be cleaned up when we switch
 	// preview worlds (in UWorld::CleanupWorld), but we want this to stick around while Persona exists
@@ -2098,14 +2098,14 @@ void FPersona::DuplicateAndSelectSocket( const FSelectedSocketInfo& SocketInfoTo
 		TargetSkeleton->Modify();
 		bModifiedSkeleton = true;
 
-		NewSocket = ConstructObject<USkeletalMeshSocket>( USkeletalMeshSocket::StaticClass(), TargetSkeleton );
+		NewSocket = NewObject<USkeletalMeshSocket>(TargetSkeleton);
 		check(NewSocket);
 	}
 	else if ( Mesh )
 	{
 		Mesh->Modify();
 
-		NewSocket = ConstructObject<USkeletalMeshSocket>( USkeletalMeshSocket::StaticClass(), Mesh );
+		NewSocket = NewObject<USkeletalMeshSocket>(Mesh);
 		check(NewSocket);
 	}
 	else 
@@ -2187,7 +2187,7 @@ bool FPersona::DoesSocketAlreadyExist( const USkeletalMeshSocket* InSocket, cons
 
 USkeletalMeshComponent*	FPersona::CreateNewSkeletalMeshComponent()
 {
-	USkeletalMeshComponent* NewComp = ConstructObject<USkeletalMeshComponent>(USkeletalMeshComponent::StaticClass());
+	USkeletalMeshComponent* NewComp = NewObject<USkeletalMeshComponent>();
 	AdditionalMeshes.Add(NewComp);
 	NewComp->SetMasterPoseComponent(GetPreviewMeshComponent());
 	return NewComp;
@@ -2255,7 +2255,7 @@ bool FPersona::AttachObjectToPreviewComponent( UObject* Object, FName AttachTo,
 		WorldSettings->SetFlags(RF_Transactional);
 		WorldSettings->Modify();
 
-		USceneComponent* SceneComponent = ConstructObject<USceneComponent>(ComponentClass, WorldSettings, NAME_None, RF_Transactional);
+		USceneComponent* SceneComponent = NewObject<USceneComponent>(WorldSettings, ComponentClass, NAME_None, RF_Transactional);
 
 		FComponentAssetBrokerage::AssignAssetToComponent(SceneComponent, Object);
 
@@ -3016,7 +3016,7 @@ void FPersona::OnImportAsset(enum EFBXImportType DefaultImportType)
 	{
 		AssetPath = NewAnimDlg->GetAssetPath();
 	
-		UFbxImportUI* ImportUI = ConstructObject<UFbxImportUI>(UFbxImportUI::StaticClass());
+		UFbxImportUI* ImportUI = NewObject<UFbxImportUI>();
 		ImportUI->Skeleton = TargetSkeleton;
 		ImportUI->MeshTypeToImport = DefaultImportType;
 
diff --git a/Engine/Source/Editor/Persona/Private/SAnimNotifyPanel.cpp b/Engine/Source/Editor/Persona/Private/SAnimNotifyPanel.cpp
--- a/Engine/Source/Editor/Persona/Private/SAnimNotifyPanel.cpp
+++ b/Engine/Source/Editor/Persona/Private/SAnimNotifyPanel.cpp
@@ -1913,7 +1913,7 @@ void SAnimNotifyTrack::CreateNewNotifyAtCursor(FString NewNotifyName, UClass* No
 
 	if( NotifyClass )
 	{
-		class UObject* AnimNotifyClass = ConstructObject<UObject>(NotifyClass, Sequence, NAME_None, RF_Transactional);
+		class UObject* AnimNotifyClass = NewObject<UObject>(Sequence, NotifyClass, NAME_None, RF_Transactional);
 		NewEvent.NotifyStateClass = Cast<UAnimNotifyState>(AnimNotifyClass);
 		NewEvent.Notify = Cast<UAnimNotify>(AnimNotifyClass);
 
diff --git a/Engine/Source/Editor/Persona/Private/SAnimationSequenceBrowser.cpp b/Engine/Source/Editor/Persona/Private/SAnimationSequenceBrowser.cpp
--- a/Engine/Source/Editor/Persona/Private/SAnimationSequenceBrowser.cpp
+++ b/Engine/Source/Editor/Persona/Private/SAnimationSequenceBrowser.cpp
@@ -855,7 +855,7 @@ void SAnimationSequenceBrowser::CreateAssetTooltipResources()
 
 	ViewportClient = MakeShareable(new FAnimationAssetViewportClient(PreviewScene));
 	SceneViewport = MakeShareable(new FSceneViewport(ViewportClient.Get(), ViewportWidget));
-	PreviewComponent = ConstructObject<UDebugSkelMeshComponent>(UDebugSkelMeshComponent::StaticClass());
+	PreviewComponent = NewObject<UDebugSkelMeshComponent>();
 
 	// Client options
 	ViewportClient->ViewportType = LVT_Perspective;
diff --git a/Engine/Source/Editor/Persona/Private/SSkeletonTree.cpp b/Engine/Source/Editor/Persona/Private/SSkeletonTree.cpp
--- a/Engine/Source/Editor/Persona/Private/SSkeletonTree.cpp
+++ b/Engine/Source/Editor/Persona/Private/SSkeletonTree.cpp
@@ -1791,7 +1791,7 @@ void SSkeletonTree::OnAddSocket()
 		const FScopedTransaction Transaction( LOCTEXT( "AddSocket", "Add Socket to Skeleton" ) );
 		TargetSkeleton->Modify();
 
-		NewSocket = ConstructObject<USkeletalMeshSocket>( USkeletalMeshSocket::StaticClass(), TargetSkeleton );
+		NewSocket = NewObject<USkeletalMeshSocket>(TargetSkeleton);
 		check(NewSocket);
 
 		NewSocket->BoneName = *static_cast<FName*>( TreeSelection.GetSingleSelectedItem()->GetData() );
@@ -1826,7 +1826,7 @@ void SSkeletonTree::OnCustomizeSocket()
 				const FScopedTransaction Transaction( LOCTEXT( "CreateMeshSocket", "Create Mesh Socket" ) );
 				Mesh->Modify();
 
-				USkeletalMeshSocket* NewSocket = ConstructObject<USkeletalMeshSocket>( USkeletalMeshSocket::StaticClass(), Mesh );
+				USkeletalMeshSocket* NewSocket = NewObject<USkeletalMeshSocket>(Mesh);
 				check(NewSocket);
 
 				NewSocket->BoneName = SocketToCustomize->BoneName;
@@ -1857,7 +1857,7 @@ void SSkeletonTree::OnPromoteSocket()
 		const FScopedTransaction Transaction( LOCTEXT( "PromoteSocket", "Promote Socket" ) );
 		TargetSkeleton->Modify();
 
-		USkeletalMeshSocket* NewSocket = ConstructObject<USkeletalMeshSocket>( USkeletalMeshSocket::StaticClass(), TargetSkeleton );
+		USkeletalMeshSocket* NewSocket = NewObject<USkeletalMeshSocket>(TargetSkeleton);
 		check(NewSocket);
 
 		NewSocket->BoneName = SocketToCustomize->BoneName;
diff --git a/Engine/Source/Editor/Persona/Private/Shared/AnimationRecorder.cpp b/Engine/Source/Editor/Persona/Private/Shared/AnimationRecorder.cpp
--- a/Engine/Source/Editor/Persona/Private/Shared/AnimationRecorder.cpp
+++ b/Engine/Source/Editor/Persona/Private/Shared/AnimationRecorder.cpp
@@ -67,7 +67,7 @@ bool FAnimationRecorder::TriggerRecordAnimation(USkeletalMeshComponent * Compone
 		}
 
 		// If not, create new one now.
-		UAnimSequence * NewSeq = ConstructObject<UAnimSequence>(UAnimSequence::StaticClass(), Parent, *AssetName, RF_Public|RF_Standalone);
+		UAnimSequence * NewSeq = NewObject<UAnimSequence>(Parent, *AssetName, RF_Public | RF_Standalone);
 		if(NewSeq)
 		{
 			// set skeleton
diff --git a/Engine/Source/Editor/PhAT/Private/PhATSharedData.cpp b/Engine/Source/Editor/PhAT/Private/PhATSharedData.cpp
--- a/Engine/Source/Editor/PhAT/Private/PhATSharedData.cpp
+++ b/Engine/Source/Editor/PhAT/Private/PhATSharedData.cpp
@@ -59,7 +59,7 @@ FPhATSharedData::FPhATSharedData()
 	MouseHandle = NewObject<UPhysicsHandleComponent>();
 
 	// Construct sim options.
-	EditorSimOptions = ConstructObject<UPhATSimOptions>(UPhATSimOptions::StaticClass());
+	EditorSimOptions = NewObject<UPhATSimOptions>(GetTransientPackage(), TEXT("EditorSimOptions"));
 	check(EditorSimOptions);
 
 	EditorSimOptions->HandleLinearDamping = MouseHandle->LinearDamping;
@@ -107,7 +107,7 @@ void FPhATSharedData::Initialize()
 	EditorSkelMesh = PreviewMesh;
 
 	// Create SkeletalMeshComponent for rendering skeletal mesh
-	EditorSkelComp = ConstructObject<UPhATEdSkeletalMeshComponent>(UPhATEdSkeletalMeshComponent::StaticClass());
+	EditorSkelComp = NewObject<UPhATEdSkeletalMeshComponent>();
 	EditorSkelComp->SharedData = this;
 	
 	// first disable collision first to avoid creating physics body
@@ -118,7 +118,7 @@ void FPhATSharedData::Initialize()
 	UStaticMesh* FloorMesh = LoadObject<UStaticMesh>(NULL, TEXT("/Engine/EditorMeshes/PhAT_FloorBox.PhAT_FloorBox"), NULL, LOAD_None, NULL);
 	check(FloorMesh);
 
-	EditorFloorComp = ConstructObject<UStaticMeshComponent>(UStaticMeshComponent::StaticClass());
+	EditorFloorComp = NewObject<UStaticMeshComponent>();
 	EditorFloorComp->StaticMesh = FloorMesh;
 	EditorFloorComp->SetRelativeScale3D(FVector(4.f));
 
diff --git a/Engine/Source/Editor/ReferenceViewer/Private/SReferenceViewer.cpp b/Engine/Source/Editor/ReferenceViewer/Private/SReferenceViewer.cpp
--- a/Engine/Source/Editor/ReferenceViewer/Private/SReferenceViewer.cpp
+++ b/Engine/Source/Editor/ReferenceViewer/Private/SReferenceViewer.cpp
@@ -34,7 +34,7 @@ void SReferenceViewer::Construct( const FArguments& InArgs )
 	HistoryManager.SetOnUpdateHistoryData(FOnUpdateHistoryData::CreateSP(this, &SReferenceViewer::OnUpdateHistoryData));
 
 	// Create the graph
-	GraphObj = ConstructObject<UEdGraph_ReferenceViewer>(UEdGraph_ReferenceViewer::StaticClass());
+	GraphObj = NewObject<UEdGraph_ReferenceViewer>();
 	GraphObj->Schema = UReferenceViewerSchema::StaticClass();
 	GraphObj->AddToRoot();
 
diff --git a/Engine/Source/Editor/SoundClassEditor/Private/SoundClassEditor.cpp b/Engine/Source/Editor/SoundClassEditor/Private/SoundClassEditor.cpp
--- a/Engine/Source/Editor/SoundClassEditor/Private/SoundClassEditor.cpp
+++ b/Engine/Source/Editor/SoundClassEditor/Private/SoundClassEditor.cpp
@@ -329,7 +329,7 @@ void FSoundClassEditor::CreateSoundClass(UEdGraphPin* FromPin, const FVector2D&
 	{
 		FName NewClassName = *Name;
 		UPackage* Package = SoundClass->GetOutermost();
-		USoundClass* NewSoundClass = ConstructObject<USoundClass>( USoundClass::StaticClass(), Package, NewClassName, RF_Public );
+		USoundClass* NewSoundClass = NewObject<USoundClass>(Package, NewClassName, RF_Public);
 
 		SoundClass->SoundClassGraph->AddNewSoundClass(FromPin, NewSoundClass, Location.X, Location.Y);
 
diff --git a/Engine/Source/Editor/StaticMeshEditor/Private/SStaticMeshEditorViewport.cpp b/Engine/Source/Editor/StaticMeshEditor/Private/SStaticMeshEditorViewport.cpp
--- a/Engine/Source/Editor/StaticMeshEditor/Private/SStaticMeshEditorViewport.cpp
+++ b/Engine/Source/Editor/StaticMeshEditor/Private/SStaticMeshEditorViewport.cpp
@@ -45,8 +45,7 @@ void SStaticMeshEditorViewport::Construct(const FArguments& InArgs)
 
 	SEditorViewport::Construct( SEditorViewport::FArguments() );
 
-	PreviewMeshComponent = ConstructObject<UStaticMeshComponent>(
-		UStaticMeshComponent::StaticClass(), GetTransientPackage(), NAME_None, RF_Transient );
+	PreviewMeshComponent = NewObject<UStaticMeshComponent>(GetTransientPackage(), NAME_None, RF_Transient );
 
 	SetPreviewMesh(StaticMesh);
 
@@ -156,7 +155,7 @@ void SStaticMeshEditorViewport::UpdatePreviewSocketMeshes()
 				// Handle adding a new component
 				if(i >= SocketedComponentCount)
 				{
-					SocketPreviewMeshComponent = ConstructObject<UStaticMeshComponent>(UStaticMeshComponent::StaticClass());
+					SocketPreviewMeshComponent = NewObject<UStaticMeshComponent>();
 					PreviewScene.AddComponent(SocketPreviewMeshComponent, FTransform::Identity);
 					SocketPreviewMeshComponents.Add(SocketPreviewMeshComponent);
 				}
@@ -205,7 +204,7 @@ void SStaticMeshEditorViewport::UpdatePreviewMesh(UStaticMesh* InStaticMesh)
 		PreviewMeshComponent = NULL;
 	}
 
-	PreviewMeshComponent = ConstructObject<UStaticMeshComponent>(UStaticMeshComponent::StaticClass());
+	PreviewMeshComponent = NewObject<UStaticMeshComponent>();
 
 	PreviewMeshComponent->SetStaticMesh(InStaticMesh);
 	PreviewScene.AddComponent(PreviewMeshComponent,FTransform::Identity);
@@ -219,7 +218,7 @@ void SStaticMeshEditorViewport::UpdatePreviewMesh(UStaticMesh* InStaticMesh)
 		UStaticMeshComponent* SocketPreviewMeshComponent = NULL;
 		if( Socket && Socket->PreviewStaticMesh )
 		{
-			SocketPreviewMeshComponent = ConstructObject<UStaticMeshComponent>(UStaticMeshComponent::StaticClass());
+			SocketPreviewMeshComponent = NewObject<UStaticMeshComponent>();
 			SocketPreviewMeshComponent->SetStaticMesh(Socket->PreviewStaticMesh);
 			SocketPreviewMeshComponent->SnapTo(PreviewMeshComponent, Socket->SocketName);
 			SocketPreviewMeshComponents.Add(SocketPreviewMeshComponent);
diff --git a/Engine/Source/Editor/StatsViewer/Private/StatsPages/CookerStatsPage.cpp b/Engine/Source/Editor/StatsViewer/Private/StatsPages/CookerStatsPage.cpp
--- a/Engine/Source/Editor/StatsViewer/Private/StatsPages/CookerStatsPage.cpp
+++ b/Engine/Source/Editor/StatsViewer/Private/StatsPages/CookerStatsPage.cpp
@@ -49,7 +49,7 @@ void FCookerStatsPage::Generate( TArray<TWeakObjectPtr<UObject>>& OutObjects ) c
 
 				if (OriginalFileName.EndsWith(TEXT(".uasset")) || OriginalFileName.EndsWith(TEXT(".umap")))
 				{
-					UCookerStats* Entry = ConstructObject<UCookerStats>(UCookerStats::StaticClass());
+					UCookerStats* Entry = NewObject<UCookerStats>();
 					Entry->AddToRoot();
 					Entry->Path = FilenameOrDirectory;
 					Entry->SizeAfter = (int32)IFileManager::Get().FileSize(FilenameOrDirectory) / 1024.0f;
@@ -102,7 +102,7 @@ void FCookerStatsPage::GenerateTotals( const TArray<TWeakObjectPtr<UObject>>& In
 		return;
 	}
 
-	UCookerStats* TotalEntry = ConstructObject<UCookerStats>(UCookerStats::StaticClass());
+	UCookerStats* TotalEntry = NewObject<UCookerStats>();
 	{
 		for (auto It = InObjects.CreateConstIterator(); It; ++It)
 		{
diff --git a/Engine/Source/Editor/StatsViewer/Private/StatsPages/LightingBuildInfoStatsPage.cpp b/Engine/Source/Editor/StatsViewer/Private/StatsPages/LightingBuildInfoStatsPage.cpp
--- a/Engine/Source/Editor/StatsViewer/Private/StatsPages/LightingBuildInfoStatsPage.cpp
+++ b/Engine/Source/Editor/StatsViewer/Private/StatsPages/LightingBuildInfoStatsPage.cpp
@@ -61,7 +61,7 @@ void FLightingBuildInfoStatsPage::GenerateTotals( const TArray< TWeakObjectPtr<U
 {
 	if(InObjects.Num())
 	{
-		ULightingBuildInfo* TotalEntry = ConstructObject<ULightingBuildInfo>( ULightingBuildInfo::StaticClass() );
+		ULightingBuildInfo* TotalEntry = NewObject<ULightingBuildInfo>();
 		for( auto It = InObjects.CreateConstIterator(); It; ++It )
 		{
 			ULightingBuildInfo* Entry = Cast<ULightingBuildInfo>( It->Get() );
diff --git a/Engine/Source/Editor/StatsViewer/Private/StatsPages/PrimitiveStatsPage.cpp b/Engine/Source/Editor/StatsViewer/Private/StatsPages/PrimitiveStatsPage.cpp
--- a/Engine/Source/Editor/StatsViewer/Private/StatsPages/PrimitiveStatsPage.cpp
+++ b/Engine/Source/Editor/StatsViewer/Private/StatsPages/PrimitiveStatsPage.cpp
@@ -258,7 +258,7 @@ struct PrimitiveStatsGenerator
 			else
 			{
 				// We don't. Create new base entry.
-				UPrimitiveStats* NewStatsEntry = ConstructObject<UPrimitiveStats>( UPrimitiveStats::StaticClass() );
+				UPrimitiveStats* NewStatsEntry = NewObject<UPrimitiveStats>();
 				NewStatsEntry->AddToRoot();
 				NewStatsEntry->Object			= Resource;
 				NewStatsEntry->Actors.AddUnique(ActorOuter);
@@ -412,7 +412,7 @@ void FPrimitiveStatsPage::GenerateTotals( const TArray< TWeakObjectPtr<UObject>
 {
 	if(InObjects.Num())
 	{
-		UPrimitiveStats* TotalEntry = ConstructObject<UPrimitiveStats>( UPrimitiveStats::StaticClass() );
+		UPrimitiveStats* TotalEntry = NewObject<UPrimitiveStats>();
 
 		TotalEntry->RadiusMin = FLT_MAX;
 		TotalEntry->RadiusMax = 0.0f;
diff --git a/Engine/Source/Editor/StatsViewer/Private/StatsPages/StaticMeshLightingInfoStatsPage.cpp b/Engine/Source/Editor/StatsViewer/Private/StatsPages/StaticMeshLightingInfoStatsPage.cpp
--- a/Engine/Source/Editor/StatsViewer/Private/StatsPages/StaticMeshLightingInfoStatsPage.cpp
+++ b/Engine/Source/Editor/StatsViewer/Private/StatsPages/StaticMeshLightingInfoStatsPage.cpp
@@ -103,7 +103,7 @@ struct StaticMeshLightingInfoStatsGenerator
 				bTextureMapping, 
 				bHasLightmapTexCoords) )
 		{
-			UStaticMeshLightingInfo* Entry = ConstructObject<UStaticMeshLightingInfo>(UStaticMeshLightingInfo::StaticClass());
+			UStaticMeshLightingInfo* Entry = NewObject<UStaticMeshLightingInfo>();
 			Entry->AddToRoot();
 			OutObjects.Add(Entry);
 
@@ -205,7 +205,7 @@ void FStaticMeshLightingInfoStatsPage::GenerateTotals( const TArray< TWeakObject
 {
 	if(InObjects.Num())
 	{
-		UStaticMeshLightingInfo* TotalEntry = ConstructObject<UStaticMeshLightingInfo>( UStaticMeshLightingInfo::StaticClass() );
+		UStaticMeshLightingInfo* TotalEntry = NewObject<UStaticMeshLightingInfo>();
 
 		for( auto It = InObjects.CreateConstIterator(); It; ++It )
 		{
diff --git a/Engine/Source/Editor/StatsViewer/Private/StatsPages/TextureStatsPage.cpp b/Engine/Source/Editor/StatsViewer/Private/StatsPages/TextureStatsPage.cpp
--- a/Engine/Source/Editor/StatsViewer/Private/StatsPages/TextureStatsPage.cpp
+++ b/Engine/Source/Editor/StatsViewer/Private/StatsPages/TextureStatsPage.cpp
@@ -185,7 +185,7 @@ struct TextureStatsGenerator : public FFindReferencedAssets
 		UTextureStats** EntryPtr = EntryMap.Find(InTexture->GetPathName());
 		if(EntryPtr == NULL)
 		{
-			Entry = ConstructObject<UTextureStats>(UTextureStats::StaticClass());
+			Entry = NewObject<UTextureStats>();
 			Entry->AddToRoot();
 			OutObjects.Add(Entry);
 			EntryMap.Add(InTexture->GetPathName(), Entry);
@@ -332,7 +332,7 @@ void FTextureStatsPage::GenerateTotals( const TArray< TWeakObjectPtr<UObject> >&
 {
 	if(InObjects.Num())
 	{
-		UTextureStats* TotalEntry = ConstructObject<UTextureStats>( UTextureStats::StaticClass() );
+		UTextureStats* TotalEntry = NewObject<UTextureStats>();
 
 		for( auto It = InObjects.CreateConstIterator(); It; ++It )
 		{
diff --git a/Engine/Source/Editor/UMGEditor/Private/Customizations/WidgetNavigationCustomization.cpp b/Engine/Source/Editor/UMGEditor/Private/Customizations/WidgetNavigationCustomization.cpp
--- a/Engine/Source/Editor/UMGEditor/Private/Customizations/WidgetNavigationCustomization.cpp
+++ b/Engine/Source/Editor/UMGEditor/Private/Customizations/WidgetNavigationCustomization.cpp
@@ -68,7 +68,7 @@ FReply FWidgetNavigationCustomization::OnCustomizeNavigation(TWeakPtr<IPropertyH
 			if ( !Widget->Navigation )
 			{
 				Widget->Modify();
-				Widget->Navigation = ConstructObject<UWidgetNavigation>(UWidgetNavigation::StaticClass(), Widget);
+				Widget->Navigation = NewObject<UWidgetNavigation>(Widget);
 			}
 		}
 	}
@@ -193,7 +193,7 @@ void FWidgetNavigationCustomization::SetNav(UWidget* Widget, EUINavigation Nav,
 	UWidgetNavigation* WidgetNavigation = Widget->Navigation;
 	if (!Widget->Navigation)
 	{
-		WidgetNavigation = ConstructObject<UWidgetNavigation>(UWidgetNavigation::StaticClass(), Widget);
+		WidgetNavigation = NewObject<UWidgetNavigation>(Widget);
 	}
 
 	switch (Nav)
diff --git a/Engine/Source/Editor/UMGEditor/Private/TabFactory/AnimationTabSummoner.cpp b/Engine/Source/Editor/UMGEditor/Private/TabFactory/AnimationTabSummoner.cpp
--- a/Engine/Source/Editor/UMGEditor/Private/TabFactory/AnimationTabSummoner.cpp
+++ b/Engine/Source/Editor/UMGEditor/Private/TabFactory/AnimationTabSummoner.cpp
@@ -311,8 +311,8 @@ private:
 	{
 		UWidgetBlueprint* WidgetBlueprint = BlueprintEditor.Pin()->GetWidgetBlueprintObj();
 
-		UWidgetAnimation* NewAnimation = ConstructObject<UWidgetAnimation>(UWidgetAnimation::StaticClass(), WidgetBlueprint, MakeUniqueObjectName(WidgetBlueprint, UWidgetAnimation::StaticClass(), "NewAnimation"), RF_Transactional);
-		NewAnimation->MovieScene =  ConstructObject<UMovieScene>(UMovieScene::StaticClass(), NewAnimation, NewAnimation->GetFName(), RF_Transactional);
+		UWidgetAnimation* NewAnimation = NewObject<UWidgetAnimation>(WidgetBlueprint, MakeUniqueObjectName(WidgetBlueprint, UWidgetAnimation::StaticClass(), "NewAnimation"), RF_Transactional);
+		NewAnimation->MovieScene = NewObject<UMovieScene>(NewAnimation, NewAnimation->GetFName(), RF_Transactional);
 
 		bool bRequestRename = true;
 		bool bNewAnimation = true;
diff --git a/Engine/Source/Editor/UMGEditor/Private/WidgetBlueprint.cpp b/Engine/Source/Editor/UMGEditor/Private/WidgetBlueprint.cpp
--- a/Engine/Source/Editor/UMGEditor/Private/WidgetBlueprint.cpp
+++ b/Engine/Source/Editor/UMGEditor/Private/WidgetBlueprint.cpp
@@ -490,7 +490,7 @@ bool FWidgetAnimation_DEPRECATED::SerializeFromMismatchedTag(struct FPropertyTag
 UWidgetBlueprint::UWidgetBlueprint(const FObjectInitializer& ObjectInitializer)
 	: Super(ObjectInitializer)
 {
-	WidgetTree = ConstructObject<UWidgetTree>(UWidgetTree::StaticClass(), this);
+	WidgetTree = CreateDefaultSubobject<UWidgetTree>(TEXT("WidgetTree"));
 	WidgetTree->SetFlags(RF_Transactional);
 }
 
@@ -512,7 +512,7 @@ void UWidgetBlueprint::PostLoad()
 			// Rename the old movie scene so we can reuse the name
 			OldAnim.MovieScene->Rename( *MakeUniqueObjectName( this, UMovieScene::StaticClass(), "MovieScene").ToString(), nullptr, REN_ForceNoResetLoaders | REN_DontCreateRedirectors | REN_DoNotDirty | REN_NonTransactional);
 
-			UWidgetAnimation* NewAnimation = ConstructObject<UWidgetAnimation>( UWidgetAnimation::StaticClass(), this, AnimName, RF_Transactional );
+			UWidgetAnimation* NewAnimation = NewObject<UWidgetAnimation>(this, AnimName, RF_Transactional);
 
 			OldAnim.MovieScene->Rename(*AnimName.ToString(), NewAnimation, REN_ForceNoResetLoaders | REN_DontCreateRedirectors | REN_DoNotDirty | REN_NonTransactional );
 
diff --git a/Engine/Source/Editor/UMGEditor/Private/WidgetBlueprintCompiler.cpp b/Engine/Source/Editor/UMGEditor/Private/WidgetBlueprintCompiler.cpp
--- a/Engine/Source/Editor/UMGEditor/Private/WidgetBlueprintCompiler.cpp
+++ b/Engine/Source/Editor/UMGEditor/Private/WidgetBlueprintCompiler.cpp
@@ -310,7 +310,7 @@ void FWidgetBlueprintCompiler::SpawnNewClass(const FString& NewClassName)
 
 	if ( NewWidgetBlueprintClass == nullptr )
 	{
-		NewWidgetBlueprintClass = ConstructObject<UWidgetBlueprintGeneratedClass>(UWidgetBlueprintGeneratedClass::StaticClass(), Blueprint->GetOutermost(), FName(*NewClassName), RF_Public | RF_Transactional);
+		NewWidgetBlueprintClass = NewObject<UWidgetBlueprintGeneratedClass>(Blueprint->GetOutermost(), FName(*NewClassName), RF_Public | RF_Transactional);
 	}
 	else
 	{
diff --git a/Engine/Source/Editor/UMGEditor/Private/WidgetBlueprintEditor.cpp b/Engine/Source/Editor/UMGEditor/Private/WidgetBlueprintEditor.cpp
--- a/Engine/Source/Editor/UMGEditor/Private/WidgetBlueprintEditor.cpp
+++ b/Engine/Source/Editor/UMGEditor/Private/WidgetBlueprintEditor.cpp
@@ -548,8 +548,8 @@ TSharedPtr<ISequencer>& FWidgetBlueprintEditor::GetSequencer()
 		}
 		else
 		{
-			UWidgetAnimation* NewAnimation = ConstructObject<UWidgetAnimation>(UWidgetAnimation::StaticClass(), Blueprint, MakeUniqueObjectName(Blueprint, UWidgetAnimation::StaticClass(), "NewAnimation"), RF_Transactional);
-			NewAnimation->MovieScene =  ConstructObject<UMovieScene>(UMovieScene::StaticClass(), NewAnimation, NAME_None, RF_Transactional);
+			UWidgetAnimation* NewAnimation = NewObject<UWidgetAnimation>(Blueprint, MakeUniqueObjectName(Blueprint, UWidgetAnimation::StaticClass(), "NewAnimation"), RF_Transactional);
+			NewAnimation->MovieScene = NewObject<UMovieScene>(NewAnimation, NAME_None, RF_Transactional);
 
 			WidgetAnimation = NewAnimation;
 		}
diff --git a/Engine/Source/Editor/UMGEditor/Private/WidgetBlueprintEditorUtils.cpp b/Engine/Source/Editor/UMGEditor/Private/WidgetBlueprintEditorUtils.cpp
--- a/Engine/Source/Editor/UMGEditor/Private/WidgetBlueprintEditorUtils.cpp
+++ b/Engine/Source/Editor/UMGEditor/Private/WidgetBlueprintEditorUtils.cpp
@@ -618,7 +618,7 @@ void FWidgetBlueprintEditorUtils::ImportWidgetsFromText(UWidgetBlueprint* BP, co
 {
 	// We create our own transient package here so that we can deserialize the data in isolation and ensure unreferenced
 	// objects not part of the deserialization set are unresolved.
-	UPackage* TempPackage = ConstructObject<UPackage>(UPackage::StaticClass(), nullptr, TEXT("/Engine/UMG/Editor/Transient"), RF_Transient);
+	UPackage* TempPackage = NewObject<UPackage>(nullptr, TEXT("/Engine/UMG/Editor/Transient"), RF_Transient);
 	TempPackage->AddToRoot();
 
 	// Turn the text buffer into objects
diff --git a/Engine/Source/Editor/UnrealEd/Private/Animation/DebugSkelMeshComponent.cpp b/Engine/Source/Editor/UnrealEd/Private/Animation/DebugSkelMeshComponent.cpp
--- a/Engine/Source/Editor/UnrealEd/Private/Animation/DebugSkelMeshComponent.cpp
+++ b/Engine/Source/Editor/UnrealEd/Private/Animation/DebugSkelMeshComponent.cpp
@@ -277,7 +277,7 @@ void UDebugSkelMeshComponent::InitAnim(bool bForceReinit)
 	// if PreviewInstance is NULL, create here once
 	if (PreviewInstance == NULL)
 	{
-		PreviewInstance = ConstructObject<UAnimPreviewInstance>(UAnimPreviewInstance::StaticClass(), this);				
+		PreviewInstance = NewObject<UAnimPreviewInstance>(this);
 		check(PreviewInstance);
 
 		//Set transactional flag in order to restore slider position when undo operation is performed
diff --git a/Engine/Source/Editor/UnrealEd/Private/AnimationCompressionPanel.cpp b/Engine/Source/Editor/UnrealEd/Private/AnimationCompressionPanel.cpp
--- a/Engine/Source/Editor/UnrealEd/Private/AnimationCompressionPanel.cpp
+++ b/Engine/Source/Editor/UnrealEd/Private/AnimationCompressionPanel.cpp
@@ -48,7 +48,7 @@ void SAnimationCompressionPanel::Construct(const FArguments& InArgs)
 		{
 			if ( Class->IsChildOf(UAnimCompress::StaticClass()) )
 			{
-				UAnimCompress* NewAlgorithm = ConstructObject<UAnimCompress>( Class );
+				UAnimCompress* NewAlgorithm = NewObject<UAnimCompress>(GetTransientPackage(), Class);
 				AnimationCompressionAlgorithms.Add( NewAlgorithm );
 			}
 		}
diff --git a/Engine/Source/Editor/UnrealEd/Private/AnimationEditorUtils.cpp b/Engine/Source/Editor/UnrealEd/Private/AnimationEditorUtils.cpp
--- a/Engine/Source/Editor/UnrealEd/Private/AnimationEditorUtils.cpp
+++ b/Engine/Source/Editor/UnrealEd/Private/AnimationEditorUtils.cpp
@@ -284,7 +284,7 @@ namespace AnimationEditorUtils
 				FString PackageName;
 				CreateUniqueAssetName(Object->GetOutermost()->GetName(), DefaultSuffix, PackageName, Name);
 
-				UAnimBlueprintFactory* Factory = ConstructObject<UAnimBlueprintFactory>(UAnimBlueprintFactory::StaticClass());
+				UAnimBlueprintFactory* Factory = NewObject<UAnimBlueprintFactory>();
 				Factory->TargetSkeleton = Object;
 
 				FContentBrowserModule& ContentBrowserModule = FModuleManager::LoadModuleChecked<FContentBrowserModule>("ContentBrowser");
@@ -305,7 +305,7 @@ namespace AnimationEditorUtils
 					CreateUniqueAssetName(Object->GetOutermost()->GetName(), DefaultSuffix, PackageName, Name);
 
 					// Create the anim blueprint factory used to generate the asset
-					UAnimBlueprintFactory* Factory = ConstructObject<UAnimBlueprintFactory>(UAnimBlueprintFactory::StaticClass());
+					UAnimBlueprintFactory* Factory = NewObject<UAnimBlueprintFactory>();
 					Factory->TargetSkeleton = Object;
 
 					FAssetToolsModule& AssetToolsModule = FModuleManager::GetModuleChecked<FAssetToolsModule>("AssetTools");
diff --git a/Engine/Source/Editor/UnrealEd/Private/Commandlets/CookCommandlet.cpp b/Engine/Source/Editor/UnrealEd/Private/Commandlets/CookCommandlet.cpp
--- a/Engine/Source/Editor/UnrealEd/Private/Commandlets/CookCommandlet.cpp
+++ b/Engine/Source/Editor/UnrealEd/Private/Commandlets/CookCommandlet.cpp
@@ -67,7 +67,7 @@ UCookCommandlet::UCookCommandlet( const FObjectInitializer& ObjectInitializer )
 
 bool UCookCommandlet::CookOnTheFly( FGuid InstanceId, int32 Timeout, bool bForceClose )
 {
-	UCookOnTheFlyServer *CookOnTheFlyServer = ConstructObject<UCookOnTheFlyServer>( UCookOnTheFlyServer::StaticClass() );
+	UCookOnTheFlyServer *CookOnTheFlyServer = NewObject<UCookOnTheFlyServer>();
 
 	struct FScopeRootObject
 	{
@@ -994,7 +994,7 @@ void UCookCommandlet::GenerateLongPackageNames(TArray<FString>& FilesInPath)
 bool UCookCommandlet::NewCook( const TArray<ITargetPlatform*>& Platforms, TArray<FString>& FilesInPath )
 {
 
-	UCookOnTheFlyServer *CookOnTheFlyServer = ConstructObject<UCookOnTheFlyServer>( UCookOnTheFlyServer::StaticClass() );
+	UCookOnTheFlyServer *CookOnTheFlyServer = NewObject<UCookOnTheFlyServer>();
 
 	struct FScopeRootObject
 	{
diff --git a/Engine/Source/Editor/UnrealEd/Private/Commandlets/DumpBlueprintsInfoCommandlet.cpp b/Engine/Source/Editor/UnrealEd/Private/Commandlets/DumpBlueprintsInfoCommandlet.cpp
--- a/Engine/Source/Editor/UnrealEd/Private/Commandlets/DumpBlueprintsInfoCommandlet.cpp
+++ b/Engine/Source/Editor/UnrealEd/Private/Commandlets/DumpBlueprintsInfoCommandlet.cpp
@@ -699,7 +699,7 @@ static AActor* DumpBlueprintInfoUtils::SpawnLevelActor(UClass* ActorClass, bool
 	}
 	else if (FKismetEditorUtilities::CanCreateBlueprintOfClass(ActorClass))
 	{
-		UActorFactory* NewFactory = ConstructObject<UActorFactory>(UActorFactoryBlueprint::StaticClass());
+		UActorFactory* NewFactory = NewObject<UActorFactory>(GetTransientPackage(), UActorFactoryBlueprint::StaticClass());
 		NewFactory->AddToRoot();
 
 		UBlueprint* ActorTemplate = MakeTempBlueprint(ActorClass);
@@ -764,7 +764,7 @@ static UBlueprint* DumpBlueprintInfoUtils::MakeTempBlueprint(UClass* ParentClass
 			BlueprintClass = UAnimBlueprint::StaticClass();
 			GeneratedClass = UAnimBlueprintGeneratedClass::StaticClass();
 
-			UAnimBlueprintFactory* BlueprintFactory = ConstructObject<UAnimBlueprintFactory>(UAnimBlueprintFactory::StaticClass());
+			UAnimBlueprintFactory* BlueprintFactory = NewObject<UAnimBlueprintFactory>();
 			BlueprintFactory->ParentClass    = ParentClass;
 			BlueprintFactory->BlueprintType  = BlueprintType;
 			BlueprintFactory->TargetSkeleton = (USkeleton*)StaticLoadObject(USkeleton::StaticClass(), /*Outer =*/nullptr, TEXT("/Engine/NotForLicensees/Automation/QAAutomationtest_Assets/TEST_SkeletalMesh_Skeleton.TEST_SkeletalMesh_Skeleton"));
@@ -787,7 +787,7 @@ static UBlueprint* DumpBlueprintInfoUtils::MakeTempBlueprint(UClass* ParentClass
 		}
 		else 
 		{
-			UBlueprintFactory* BlueprintFactory = ConstructObject<UBlueprintFactory>(UBlueprintFactory::StaticClass());
+			UBlueprintFactory* BlueprintFactory = NewObject<UBlueprintFactory>();
 			BlueprintFactory->ParentClass = ParentClass;
 			AssetFactory = BlueprintFactory;
 		}
diff --git a/Engine/Source/Editor/UnrealEd/Private/Commandlets/GatherTextCommandlet.cpp b/Engine/Source/Editor/UnrealEd/Private/Commandlets/GatherTextCommandlet.cpp
--- a/Engine/Source/Editor/UnrealEd/Private/Commandlets/GatherTextCommandlet.cpp
+++ b/Engine/Source/Editor/UnrealEd/Private/Commandlets/GatherTextCommandlet.cpp
@@ -110,7 +110,7 @@ int32 UGatherTextCommandlet::Main( const FString& Params )
 			continue;
 		}
 
-		UGatherTextCommandletBase* Commandlet = ConstructObject<UGatherTextCommandletBase>( CommandletClass );
+		UGatherTextCommandletBase* Commandlet = NewObject<UGatherTextCommandletBase>(GetTransientPackage(), CommandletClass);
 		check(Commandlet);
 		Commandlet->AddToRoot();
 		Commandlet->Initialize( ManifestInfo, SourceControlInfo );
diff --git a/Engine/Source/Editor/UnrealEd/Private/Commandlets/PackageUtilities.cpp b/Engine/Source/Editor/UnrealEd/Private/Commandlets/PackageUtilities.cpp
--- a/Engine/Source/Editor/UnrealEd/Private/Commandlets/PackageUtilities.cpp
+++ b/Engine/Source/Editor/UnrealEd/Private/Commandlets/PackageUtilities.cpp
@@ -1978,7 +1978,7 @@ struct CompressAnimationsFunctor
 			if( bResetCompression )
 			{
 				UE_LOG(LogPackageUtilities, Warning, TEXT("%s (%s) Resetting with BitwiseCompressOnly."), *AnimSeq->GetName(), *AnimSeq->GetFullName());
-				UAnimCompress* CompressionAlgorithm = ConstructObject<UAnimCompress_BitwiseCompressOnly>( UAnimCompress_BitwiseCompressOnly::StaticClass() );
+				UAnimCompress* CompressionAlgorithm = NewObject<UAnimCompress_BitwiseCompressOnly>();
 				CompressionAlgorithm->RotationCompressionFormat = ACF_Float96NoW;
 				CompressionAlgorithm->TranslationCompressionFormat = ACF_None;
 				CompressionAlgorithm->ScaleCompressionFormat = ACF_Float96NoW;
diff --git a/Engine/Source/Editor/UnrealEd/Private/Editor.cpp b/Engine/Source/Editor/UnrealEd/Private/Editor.cpp
--- a/Engine/Source/Editor/UnrealEd/Private/Editor.cpp
+++ b/Engine/Source/Editor/UnrealEd/Private/Editor.cpp
@@ -204,15 +204,15 @@ static void OnObjectSelected(UObject* Object)
 
 static void PrivateInitSelectedSets()
 {
-	PrivateGetSelectedActors() = NewNamedObject<USelection>(GetTransientPackage(), TEXT("SelectedActors"), RF_Transactional);
+	PrivateGetSelectedActors() = NewObject<USelection>(GetTransientPackage(), TEXT("SelectedActors"), RF_Transactional);
 	PrivateGetSelectedActors()->AddToRoot();
 
 	PrivateGetSelectedActors()->SelectObjectEvent.AddStatic(&OnObjectSelected);
 
-	PrivateGetSelectedComponents() = NewNamedObject<USelection>(GetTransientPackage(), TEXT("SelectedComponents"), RF_Transactional);
+	PrivateGetSelectedComponents() = NewObject<USelection>(GetTransientPackage(), TEXT("SelectedComponents"), RF_Transactional);
 	PrivateGetSelectedComponents()->AddToRoot();
 
-	PrivateGetSelectedObjects() = NewNamedObject<USelection>(GetTransientPackage(), TEXT("SelectedObjects"), RF_Transactional);
+	PrivateGetSelectedObjects() = NewObject<USelection>(GetTransientPackage(), TEXT("SelectedObjects"), RF_Transactional);
 	PrivateGetSelectedObjects()->AddToRoot();
 }
 
@@ -746,7 +746,7 @@ void UEditorEngine::Init(IEngineLoop* InEngineLoop)
 					{
 						for ( int32 i=0; i < VolumeClasses.Num(); i++ )
 						{
-							UActorFactory* NewFactory = ConstructObject<UActorFactory>(TestClass);
+							UActorFactory* NewFactory = NewObject<UActorFactory>(GetTransientPackage(), TestClass);
 							check(NewFactory);
 							NewFactory->NewActorClass = VolumeClasses[i];
 							ActorFactories.Add(NewFactory);
@@ -754,7 +754,7 @@ void UEditorEngine::Init(IEngineLoop* InEngineLoop)
 					}
 					else
 					{
-						UActorFactory* NewFactory = ConstructObject<UActorFactory>(TestClass);
+						UActorFactory* NewFactory = NewObject<UActorFactory>(GetTransientPackage(), TestClass);
 						check(NewFactory);
 						ActorFactories.Add(NewFactory);
 					}
@@ -820,7 +820,7 @@ void UEditorEngine::InitBuilderBrush( UWorld* InWorld )
 
 	// For additive geometry mode, make the builder brush a small 256x256x256 cube so its visible.
 	const int32 CubeSize = 256;
-	UCubeBuilder* CubeBuilder = ConstructObject<UCubeBuilder>( UCubeBuilder::StaticClass() );
+	UCubeBuilder* CubeBuilder = NewObject<UCubeBuilder>();
 	CubeBuilder->X = CubeSize;
 	CubeBuilder->Y = CubeSize;
 	CubeBuilder->Z = CubeSize;
@@ -1642,7 +1642,7 @@ const UEditorUserSettings& UEditorEngine::GetEditorUserSettings() const
 	if (EditorUserSettings == NULL)
 	{
 		auto ConstThis = const_cast< UEditorEngine* >( this );	// Hack because Header Generator doesn't yet support mutable keyword
-		ConstThis->EditorUserSettings = ConstructObject<UEditorUserSettings>(UEditorUserSettings::StaticClass());
+		ConstThis->EditorUserSettings = NewObject<UEditorUserSettings>();
 	}
 	return *EditorUserSettings;
 }
@@ -1651,7 +1651,7 @@ UEditorUserSettings& UEditorEngine::AccessEditorUserSettings()
 {
 	if (EditorUserSettings == NULL)
 	{
-		EditorUserSettings = ConstructObject<UEditorUserSettings>(UEditorUserSettings::StaticClass());
+		EditorUserSettings = NewObject<UEditorUserSettings>();
 	}
 	return *EditorUserSettings;
 }
@@ -1669,7 +1669,7 @@ const UEditorGameAgnosticSettings& UEditorEngine::GetGameAgnosticSettings() cons
 	if (GameAgnosticSettings == NULL)
 	{
 		auto ConstThis = const_cast< UEditorEngine* >( this );	// Hack because Header Generator doesn't yet support mutable keyword
-		ConstThis->GameAgnosticSettings = ConstructObject<UEditorGameAgnosticSettings>(UEditorGameAgnosticSettings::StaticClass());
+		ConstThis->GameAgnosticSettings = NewObject<UEditorGameAgnosticSettings>();
 		
 		// Load config from file, but only if we are not the build machine since game agnostic settings may put the builder in an unclean state
 		if ( !GIsBuildMachine )
@@ -1684,7 +1684,7 @@ UEditorGameAgnosticSettings& UEditorEngine::AccessGameAgnosticSettings()
 {
 	if (GameAgnosticSettings == NULL)
 	{
-		GameAgnosticSettings = ConstructObject<UEditorGameAgnosticSettings>(UEditorGameAgnosticSettings::StaticClass());
+		GameAgnosticSettings = NewObject<UEditorGameAgnosticSettings>();
 		
 		// Load config from file, but only if we are not the build machine since game agnostic settings may put the builder in an unclean state
 		if ( !GIsBuildMachine )
@@ -1915,7 +1915,7 @@ UAudioComponent* UEditorEngine::ResetPreviewAudioComponent( USoundBase* Sound, U
 		}
 		else
 		{
-			PreviewSoundCue = ConstructObject<USoundCue>( USoundCue::StaticClass() );
+			PreviewSoundCue = NewObject<USoundCue>();
 			// Set world to NULL as it will most likely become invalid in the next PIE/Simulate session and the
 			// component will be left with invalid pointer.
 			PreviewAudioComponent = FAudioDevice::CreateComponent( PreviewSoundCue, NULL, NULL, false );
@@ -3998,7 +3998,7 @@ UBrushBuilder* UEditorEngine::FindBrushBuilder( UClass* BrushBuilderClass )
 	else
 	{
 		// An existing builder does not exist so create one now
-		Builder = ConstructObject<UBrushBuilder>( BrushBuilderClass );
+		Builder = NewObject<UBrushBuilder>(GetTransientPackage(), BrushBuilderClass);
 		BrushBuilders.Add( Builder );
 	}
 
@@ -6461,7 +6461,7 @@ bool UEditorEngine::LoadPreviewMesh( int32 Index )
 		// If we don't have a preview mesh component in the world yet, create one. 
 		if( !PreviewMeshComp )
 		{
-			PreviewMeshComp = ConstructObject<UStaticMeshComponent>( UStaticMeshComponent::StaticClass(), GetTransientPackage() );
+			PreviewMeshComp = NewObject<UStaticMeshComponent>();
 			check(PreviewMeshComp);
 
 			// Attach the component to the scene even if the preview mesh doesn't load.
diff --git a/Engine/Source/Editor/UnrealEd/Private/Editor/EditorUserSettings.cpp b/Engine/Source/Editor/UnrealEd/Private/Editor/EditorUserSettings.cpp
--- a/Engine/Source/Editor/UnrealEd/Private/Editor/EditorUserSettings.cpp
+++ b/Engine/Source/Editor/UnrealEd/Private/Editor/EditorUserSettings.cpp
@@ -12,7 +12,7 @@ UEditorUserSettings::UEditorUserSettings(const FObjectInitializer& ObjectInitial
 	//Default to high quality
 	MaterialQualityLevel = 1;
 	bMonitorEditorPerformance = true;
-	BlueprintFavorites = ConstructObject<UBlueprintPaletteFavorites>(UBlueprintPaletteFavorites::StaticClass(), this);
+	BlueprintFavorites = CreateDefaultSubobject<UBlueprintPaletteFavorites>(TEXT("BlueprintFavorites"));
 }
 
 void UEditorUserSettings::PostInitProperties()
diff --git a/Engine/Source/Editor/UnrealEd/Private/EditorActorFolders.cpp b/Engine/Source/Editor/UnrealEd/Private/EditorActorFolders.cpp
--- a/Engine/Source/Editor/UnrealEd/Private/EditorActorFolders.cpp
+++ b/Engine/Source/Editor/UnrealEd/Private/EditorActorFolders.cpp
@@ -217,7 +217,7 @@ UEditorActorFolders& FActorFolders::InitializeForWorld(UWorld& InWorld)
 
 	// We intentionally don't pass RF_Transactional to ConstructObject so that we don't record the creation of the object into the undo buffer
 	// (to stop it getting deleted on undo as we manage its lifetime), but we still want it to be RF_Transactional so we can record any changes later
-	UEditorActorFolders* Folders = ConstructObject<UEditorActorFolders>(UEditorActorFolders::StaticClass(), GetTransientPackage(), NAME_None, RF_NoFlags);
+	UEditorActorFolders* Folders = NewObject<UEditorActorFolders>(GetTransientPackage(), NAME_None, RF_NoFlags);
 	Folders->SetFlags(RF_Transactional);
 	TemporaryWorldFolders.Add(&InWorld, Folders);
 
diff --git a/Engine/Source/Editor/UnrealEd/Private/EditorLevelUtils.cpp b/Engine/Source/Editor/UnrealEd/Private/EditorLevelUtils.cpp
--- a/Engine/Source/Editor/UnrealEd/Private/EditorLevelUtils.cpp
+++ b/Engine/Source/Editor/UnrealEd/Private/EditorLevelUtils.cpp
@@ -563,7 +563,7 @@ namespace EditorLevelUtils
 		if (UEditorEngine::IsUsingWorldAssets())
 		{
 			// Create a new world
-			UWorldFactory* Factory = ConstructObject<UWorldFactory>(UWorldFactory::StaticClass());
+			UWorldFactory* Factory = NewObject<UWorldFactory>();
 			Factory->WorldType = EWorldType::Inactive;
 			UPackage* Pkg = CreatePackage(NULL, NULL);
 			FName WorldName(TEXT("NewWorld"));
diff --git a/Engine/Source/Editor/UnrealEd/Private/EditorModeManager.cpp b/Engine/Source/Editor/UnrealEd/Private/EditorModeManager.cpp
--- a/Engine/Source/Editor/UnrealEd/Private/EditorModeManager.cpp
+++ b/Engine/Source/Editor/UnrealEd/Private/EditorModeManager.cpp
@@ -912,7 +912,7 @@ void FEditorModeTools::SetBookmark( uint32 InIndex, FEditorViewportClient* InVie
 			// If the index doesn't already have a bookmark in place, create a new one
 			if ( !WorldSettings->BookMarks[ InIndex ] )
 			{
-				WorldSettings->BookMarks[ InIndex ] = ConstructObject<UBookMark>( UBookMark::StaticClass(), WorldSettings );
+				WorldSettings->BookMarks[InIndex] = NewObject<UBookMark>(WorldSettings);
 			}
 
 			UBookMark* CurBookMark = WorldSettings->BookMarks[ InIndex ];
diff --git a/Engine/Source/Editor/UnrealEd/Private/EditorObject.cpp b/Engine/Source/Editor/UnrealEd/Private/EditorObject.cpp
--- a/Engine/Source/Editor/UnrealEd/Private/EditorObject.cpp
+++ b/Engine/Source/Editor/UnrealEd/Private/EditorObject.cpp
@@ -473,9 +473,9 @@ static const TCHAR* ImportProperties(
 
 				if (!ComponentTemplate)
 				{
-					ComponentTemplate = ConstructObject<UObject>(
-						TemplateClass,
+					ComponentTemplate = NewObject<UObject>(
 						SubobjectOuter,
+						TemplateClass,
 						TemplateName,
 						NewFlags,
 						Archetype,
diff --git a/Engine/Source/Editor/UnrealEd/Private/EditorServer.cpp b/Engine/Source/Editor/UnrealEd/Private/EditorServer.cpp
--- a/Engine/Source/Editor/UnrealEd/Private/EditorServer.cpp
+++ b/Engine/Source/Editor/UnrealEd/Private/EditorServer.cpp
@@ -316,7 +316,7 @@ bool UEditorEngine::SafeExec( UWorld* InWorld, const TCHAR* InStr, FOutputDevice
 		UFactory* Factory = NULL;
 		if( Class->IsChildOf(UFactory::StaticClass()) )
 		{
-			Factory = ConstructObject<UFactory>( Class );
+			Factory = NewObject<UFactory>(GetTransientPackage(), Class);
 		}
 
 		UObject* NewObject = NULL;
@@ -1959,7 +1959,7 @@ UWorld* UEditorEngine::NewMap()
 	EditorDestroyWorld( Context, CleanseText );
 
 	// Create a new world
-	UWorldFactory* Factory = ConstructObject<UWorldFactory>(UWorldFactory::StaticClass());
+	UWorldFactory* Factory = NewObject<UWorldFactory>();
 	Factory->WorldType = EWorldType::Editor;
 	Factory->bInformEngineOfWorld = true;
 	Factory->FeatureLevel = GEditor->DefaultWorldFeatureLevel;
@@ -2827,7 +2827,7 @@ void UEditorEngine::DoMoveSelectedActorsToLevel( ULevel* InDestLevel )
 
 ULevel*  UEditorEngine::CreateTransLevelMoveBuffer( UWorld* InWorld )
 {
-	ULevel* BufferLevel = NewNamedObject<ULevel>(GetTransientPackage(), TEXT("TransLevelMoveBuffer"));
+	ULevel* BufferLevel = NewObject<ULevel>(GetTransientPackage(), TEXT("TransLevelMoveBuffer"));
 	BufferLevel->Initialize(FURL(nullptr));
 	check( BufferLevel );
 	BufferLevel->AddToRoot();
@@ -2852,7 +2852,7 @@ ULevel*  UEditorEngine::CreateTransLevelMoveBuffer( UWorld* InWorld )
 	BufferDefaultBrush->Brush = CastChecked<UModel>(StaticFindObject(UModel::StaticClass(), BufferLevel->OwningWorld->GetOuter(), TEXT("Brush"), true), ECastCheckedType::NullAllowed);
 	if (!BufferDefaultBrush->Brush)
 	{
-		BufferDefaultBrush->Brush = NewNamedObject<UModel>(InWorld, TEXT("Brush"));
+		BufferDefaultBrush->Brush = NewObject<UModel>(InWorld, TEXT("Brush"));
 		BufferDefaultBrush->Brush->Initialize(BufferDefaultBrush, 1);
 	}
 	BufferDefaultBrush->GetBrushComponent()->Brush = BufferDefaultBrush->Brush;
@@ -5513,11 +5513,11 @@ bool UEditorEngine::HandleTestPropsCommand( const TCHAR* Str, FOutputDevice& Ar
 	UClass* Class = NULL;
 	if( ParseObject<UClass>( Str, TEXT("CLASS="), Class, ANY_PACKAGE ) != false )
 	{ 
-		Object = ConstructObject<UObject>( Class );
+		Object = NewObject<UObject>(GetTransientPackage(), Class);
 	}
 	else
 	{
-		Object = ConstructObject<UPropertyEditorTestObject>( UPropertyEditorTestObject::StaticClass() );
+		Object = NewObject<UPropertyEditorTestObject>();
 	}
 
 	TSharedRef<SWindow> Window = SNew(SWindow)
@@ -5555,7 +5555,7 @@ bool UEditorEngine::HandleTestPropsCommand( const TCHAR* Str, FOutputDevice& Ar
 
 		for (int Count = 0; Count < 50; Count++)
 		{
-			Objects.Add( ConstructObject<UPropertyEditorTestObject>( UPropertyEditorTestObject::StaticClass() ) );
+			Objects.Add(NewObject<UPropertyEditorTestObject>());
 		}
 
 		Table->SetObjects( Objects );
diff --git a/Engine/Source/Editor/UnrealEd/Private/Factories/ActorFactory.cpp b/Engine/Source/Editor/UnrealEd/Private/Factories/ActorFactory.cpp
--- a/Engine/Source/Editor/UnrealEd/Private/Factories/ActorFactory.cpp
+++ b/Engine/Source/Editor/UnrealEd/Private/Factories/ActorFactory.cpp
@@ -1148,7 +1148,7 @@ UActorFactoryCameraActor::UActorFactoryCameraActor(const FObjectInitializer& Obj
 static UBillboardComponent* CreateEditorOnlyBillboardComponent(AActor* ActorOwner, USceneComponent* AttachParent)
 {
 	// Create a new billboard component to serve as a visualization of the actor until there is another primitive component
-	UBillboardComponent* BillboardComponent = ConstructObject<UBillboardComponent>(UBillboardComponent::StaticClass(), ActorOwner, NAME_None, RF_Transactional);
+	UBillboardComponent* BillboardComponent = NewObject<UBillboardComponent>(ActorOwner, NAME_None, RF_Transactional);
 
 	BillboardComponent->Sprite = LoadObject<UTexture2D>(nullptr, TEXT("/Engine/EditorResources/EmptyActor.EmptyActor"));
 	BillboardComponent->RelativeScale3D = FVector(0.5f, 0.5f, 0.5f);
@@ -1184,7 +1184,7 @@ AActor* UActorFactoryEmptyActor::SpawnActor( UObject* Asset, ULevel* InLevel, co
 		// Spawn a temporary actor for dragging around
 		NewActor = Super::SpawnActor(Asset, InLevel, Location, Rotation, ObjectFlags, Name);
 
-		USceneComponent* RootComponent = ConstructObject<USceneComponent>(USceneComponent::StaticClass(), NewActor, USceneComponent::GetDefaultSceneRootVariableName(), RF_Transactional);
+		USceneComponent* RootComponent = NewObject<USceneComponent>(NewActor, USceneComponent::GetDefaultSceneRootVariableName(), RF_Transactional);
 		RootComponent->Mobility = EComponentMobility::Movable;
 		RootComponent->SetWorldLocationAndRotation(Location, Rotation);
 		NewActor->SetRootComponent(RootComponent);
@@ -1530,7 +1530,7 @@ void UActorFactoryMatineeActor::PostSpawnActor( UObject* Asset, AActor* NewActor
 	else
 	{
 		// if MatineeData isn't set yet, create default one
-		UInterpData * NewMatineeData = ConstructObject<UInterpData>(UInterpData::StaticClass(), NewActor);
+		UInterpData * NewMatineeData = NewObject<UInterpData>(NewActor);
 		MatineeActor->MatineeData = NewMatineeData;
 	}
 }
@@ -1805,9 +1805,9 @@ void CreateBrushForVolumeActor( AVolume* NewActor, UBrushBuilder* BrushBuilder )
 		NewActor->PreEditChange(NULL);
 
 		NewActor->PolyFlags = 0;
-		NewActor->Brush = NewNamedObject<UModel>(NewActor, NAME_None, RF_Transactional);
+		NewActor->Brush = NewObject<UModel>(NewActor, NAME_None, RF_Transactional);
 		NewActor->Brush->Initialize(nullptr, true);
-		NewActor->Brush->Polys = NewNamedObject<UPolys>(NewActor->Brush, NAME_None, RF_Transactional);
+		NewActor->Brush->Polys = NewObject<UPolys>(NewActor->Brush, NAME_None, RF_Transactional);
 		NewActor->GetBrushComponent()->Brush = NewActor->Brush;
 		if(BrushBuilder != nullptr)
 		{
@@ -1863,7 +1863,7 @@ void UActorFactoryBoxVolume::PostSpawnActor( UObject* Asset, AActor* NewActor )
 	AVolume* VolumeActor = CastChecked<AVolume>(NewActor);
 	if ( VolumeActor != NULL )
 	{
-		UCubeBuilder* Builder = ConstructObject<UCubeBuilder>( UCubeBuilder::StaticClass() );
+		UCubeBuilder* Builder = NewObject<UCubeBuilder>();
 		CreateBrushForVolumeActor( VolumeActor, Builder );
 
 		if (Asset)
@@ -1903,7 +1903,7 @@ void UActorFactorySphereVolume::PostSpawnActor( UObject* Asset, AActor* NewActor
 	AVolume* VolumeActor = CastChecked<AVolume>(NewActor);
 	if ( VolumeActor != NULL )
 	{
-		UTetrahedronBuilder* Builder = ConstructObject<UTetrahedronBuilder>( UTetrahedronBuilder::StaticClass() );
+		UTetrahedronBuilder* Builder = NewObject<UTetrahedronBuilder>();
 		Builder->SphereExtrapolation = 2;
 		Builder->Radius = 192.0f;
 		CreateBrushForVolumeActor( VolumeActor, Builder );
@@ -1944,7 +1944,7 @@ void UActorFactoryCylinderVolume::PostSpawnActor( UObject* Asset, AActor* NewAct
 	AVolume* VolumeActor = CastChecked<AVolume>(NewActor);
 	if ( VolumeActor != NULL )
 	{
-		UCylinderBuilder* Builder = ConstructObject<UCylinderBuilder>( UCylinderBuilder::StaticClass() );
+		UCylinderBuilder* Builder = NewObject<UCylinderBuilder>();
 		Builder->OuterRadius = 128.0f;
 		CreateBrushForVolumeActor( VolumeActor, Builder );
 
diff --git a/Engine/Source/Editor/UnrealEd/Private/Factories/AnimCompositeFactory.cpp b/Engine/Source/Editor/UnrealEd/Private/Factories/AnimCompositeFactory.cpp
--- a/Engine/Source/Editor/UnrealEd/Private/Factories/AnimCompositeFactory.cpp
+++ b/Engine/Source/Editor/UnrealEd/Private/Factories/AnimCompositeFactory.cpp
@@ -65,7 +65,7 @@ UObject* UAnimCompositeFactory::FactoryCreateNew(UClass* Class, UObject* InParen
 {
 	if (TargetSkeleton || SourceAnimation)
 	{
-		UAnimComposite* AnimComposite = ConstructObject<UAnimComposite>(Class,InParent,Name,Flags);
+		UAnimComposite* AnimComposite = NewObject<UAnimComposite>(InParent, Class, Name, Flags);
 
 		if(SourceAnimation)
 		{
diff --git a/Engine/Source/Editor/UnrealEd/Private/Factories/AnimMontageFactory.cpp b/Engine/Source/Editor/UnrealEd/Private/Factories/AnimMontageFactory.cpp
--- a/Engine/Source/Editor/UnrealEd/Private/Factories/AnimMontageFactory.cpp
+++ b/Engine/Source/Editor/UnrealEd/Private/Factories/AnimMontageFactory.cpp
@@ -64,7 +64,7 @@ UObject* UAnimMontageFactory::FactoryCreateNew(UClass* Class, UObject* InParent,
 {
 	if (TargetSkeleton || SourceAnimation)
 	{
-		UAnimMontage* AnimMontage = ConstructObject<UAnimMontage>(Class,InParent,Name,Flags);
+		UAnimMontage* AnimMontage = NewObject<UAnimMontage>(InParent, Class, Name, Flags);
 
 		if(SourceAnimation)
 		{
diff --git a/Engine/Source/Editor/UnrealEd/Private/Factories/AnimSequenceFactory.cpp b/Engine/Source/Editor/UnrealEd/Private/Factories/AnimSequenceFactory.cpp
--- a/Engine/Source/Editor/UnrealEd/Private/Factories/AnimSequenceFactory.cpp
+++ b/Engine/Source/Editor/UnrealEd/Private/Factories/AnimSequenceFactory.cpp
@@ -63,7 +63,7 @@ UObject* UAnimSequenceFactory::FactoryCreateNew(UClass* Class, UObject* InParent
 {
 	if (TargetSkeleton)
 	{
-		UAnimSequence* AnimSequence = ConstructObject<UAnimSequence>(Class,InParent,Name,Flags);
+		UAnimSequence* AnimSequence = NewObject<UAnimSequence>(InParent, Class, Name, Flags);
 
 		// @todo I think this will crash, we should support differentoptions
 		AnimSequence->SequenceLength = 0.f;
diff --git a/Engine/Source/Editor/UnrealEd/Private/Factories/ApexDestructibleAssetImport.cpp b/Engine/Source/Editor/UnrealEd/Private/Factories/ApexDestructibleAssetImport.cpp
--- a/Engine/Source/Editor/UnrealEd/Private/Factories/ApexDestructibleAssetImport.cpp
+++ b/Engine/Source/Editor/UnrealEd/Private/Factories/ApexDestructibleAssetImport.cpp
@@ -998,7 +998,7 @@ UDestructibleMesh* ImportDestructibleMeshFromApexDestructibleAsset(UObject* InPa
 	{
 		// Store the current file path and timestamp for re-import purposes
 		// @todo AssetImportData make a data class for Apex destructible assets
-		DestructibleMesh->AssetImportData = ConstructObject<UAssetImportData>(UAssetImportData::StaticClass(), DestructibleMesh);
+		DestructibleMesh->AssetImportData = NewObject<UAssetImportData>(DestructibleMesh);
 		DestructibleMesh->AssetImportData->SourceFilePath = FReimportManager::SanitizeImportFilename(UFactory::CurrentFilename, DestructibleMesh);
 		DestructibleMesh->AssetImportData->SourceFileTimestamp = IFileManager::Get().GetTimeStamp(*UFactory::CurrentFilename).ToString();
 		DestructibleMesh->AssetImportData->bDirty = false;
