UE-7920 Using Blueprint Interface functions that take other Interfaces as parameters can cause the Editor to crash in PIE mode.

diff --git a/Engine/Source/Editor/KismetCompiler/Private/KismetCompiler.cpp b/Engine/Source/Editor/KismetCompiler/Private/KismetCompiler.cpp
--- a/Engine/Source/Editor/KismetCompiler/Private/KismetCompiler.cpp
+++ b/Engine/Source/Editor/KismetCompiler/Private/KismetCompiler.cpp
@@ -300,6 +300,17 @@ void FKismetCompilerContext::ValidateLink(const UEdGraphPin* PinA, const UEdGrap
 	{
 		MessageLog.Warning(*FString::Printf(*LOCTEXT("PinTypeMismatch_Error", "Can't connect pins @@ and @@: %s").ToString(), *ConnectResponse.Message.ToString()), PinA, PinB);
 	}
+
+	if (PinA && PinB && PinA->Direction != PinB->Direction)
+	{
+		const UEdGraphPin* InputPin = (EEdGraphPinDirection::EGPD_Input == PinA->Direction) ? PinA : PinB;
+		const UEdGraphPin* OutputPin = (EEdGraphPinDirection::EGPD_Output == PinA->Direction) ? PinA : PinB;
+		const bool bForbiddenConnection = InputPin && OutputPin && (OutputPin->PinType.PinCategory == Schema->PC_Interface) && (InputPin->PinType.PinCategory == Schema->PC_Object);
+		if (bForbiddenConnection)
+		{
+			MessageLog.Error(*LOCTEXT("PinTypeMismatch_Error", "Can't connect pins @@ (Interface) and @@ (Object). Use an explicit cast node.").ToString(), OutputPin, InputPin);
+		}
+	}
 }
 
 /** Validate that the wiring for a single pin is schema compatible */
