[INTEGRATE] Change 2440033 by [REDACTED_PERSON][REDACTED_HANDLE] on 2015/02/10 13:23:59

	Blueprints: Allow user to specify ClassGroupNames for components via the already editable BlueprintCategory class setting [UE-9305]
	- Note: Changed categories won't appear in the menu until the FComponentTypeRegistry is refreshed currently (WIP)

diff --git a/Engine/Source/Editor/KismetCompiler/Private/KismetCompiler.cpp b/Engine/Source/Editor/KismetCompiler/Private/KismetCompiler.cpp
--- a/Engine/Source/Editor/KismetCompiler/Private/KismetCompiler.cpp
+++ b/Engine/Source/Editor/KismetCompiler/Private/KismetCompiler.cpp
@@ -1722,9 +1722,15 @@ void FKismetCompilerContext::FinishCompilingClass(UClass* Class)
 		{
 			static const FName NAME_ClassGroupNames(TEXT("ClassGroupNames"));
 			Class->SetMetaData(FBlueprintMetadata::MD_BlueprintSpawnableComponent, TEXT("true"));
-			Class->SetMetaData(NAME_ClassGroupNames, *NSLOCTEXT("BlueprintableComponents", "CategoryName", "Custom").ToString());
-		}
 
+			FString ClassGroupCategory = NSLOCTEXT("BlueprintableComponents", "CategoryName", "Custom").ToString();
+			if (!Blueprint->BlueprintCategory.IsEmpty())
+			{
+				ClassGroupCategory = Blueprint->BlueprintCategory;
+			}
+
+			Class->SetMetaData(NAME_ClassGroupNames, *ClassGroupCategory);
+		}
 
 		// Add a category if one has been specified
 		if(Blueprint->BlueprintCategory.Len() > 0)
