[INTEGRATE] Change 2440043 by [REDACTED_PERSON][REDACTED_HANDLE] on 2015/02/10 13:29:52

	Blueprints: Allow any Actor or Component subclass to tick by default
	- Added a new project setting bCanBlueprintsTickByDefault, which defaults to true
	- Added a new class metadata flag ChildCannotTick, which can be used to prevent Blueprint subclasses from ticking
	- Improved the warning messages when ticking isn't possible

	- Blueprint subclasses are now allowed to tick as long as any of the following are true of the most derived C++ parent class:
	  - C++ parent has bCanEverTick = true on their primary tick function
	  - C++ parent is AActor or UActorComponent
	  - C++ parent does not have ChildCannotTick metadata and bCanBlueprintsTickByDefault is true
	  - C++ parent does not have ChildCannotTick metadata but does have ChildCanTick metadata

	Upgrade note: If you are carefully managing which Blueprints are allowed to tick (e.g., for dedicated server performance), you will probably want to set bCanBlueprintsTickByDefault to false for your project and explicitly opt in for ticking via ChildCanTick on specific cosmetic classes.
	[UE-9179]

diff --git a/Engine/Config/BaseEngine.ini b/Engine/Config/BaseEngine.ini
--- a/Engine/Config/BaseEngine.ini
+++ b/Engine/Config/BaseEngine.ini
@@ -125,6 +125,8 @@ bSubtitlesEnabled=True
 bSubtitlesForcedOff=false
 DefaultSoundName=/Engine/EngineSounds/WhiteNoise.WhiteNoise
 MaximumLoopIterationCount = 1000000
+; Allows Blueprint classes based on a native Actor or Component subclass to tick even if their parent has bCanEverTick flag set to false.
+bCanBlueprintsTickByDefault=true
 ; If setting bAllowDebugViewmodesOnConsoles=true, must also add EngineDebugMaterials to the StartupPackages for console platforms
 bAllowDebugViewmodesOnConsoles=false
 CameraRotationThreshold=45.0
diff --git a/Engine/Source/Editor/BlueprintGraph/Classes/EdGraphSchema_K2.h b/Engine/Source/Editor/BlueprintGraph/Classes/EdGraphSchema_K2.h
--- a/Engine/Source/Editor/BlueprintGraph/Classes/EdGraphSchema_K2.h
+++ b/Engine/Source/Editor/BlueprintGraph/Classes/EdGraphSchema_K2.h
@@ -59,6 +59,12 @@ public:
 	/** A listing of classes that this class is accessible from (and only those classes, if present).  Note that this determines the GRAPH CONTEXTS in which the node cannot be placed (e.g. right click menu, palette), and does NOT control menus when dragging off of a context pin (i.e. contextual drag) */
 	static const FName MD_RestrictedToClasses;
 
+	/// [ClassMetadata] Used for Actor and Component classes. If the native class cannot tick, Blueprint generated classes based this Actor or Component can have bCanEverTick flag overridden even if bCanBlueprintsTickByDefault is false.
+	static const FName MD_ChildCanTick;
+
+	/// [ClassMetadata] Used for Actor and Component classes. If the native class cannot tick, Blueprint generated classes based this Actor or Component can never tick even if bCanBlueprintsTickByDefault is true.
+	static const FName MD_ChildCannotTick;
+
 	//    function metadata
 	/** Specifies a UFUNCTION as Kismet protected, which can only be called from itself */
 	static const FName MD_Protected;
diff --git a/Engine/Source/Editor/BlueprintGraph/Private/EdGraphSchema_K2.cpp b/Engine/Source/Editor/BlueprintGraph/Private/EdGraphSchema_K2.cpp
--- a/Engine/Source/Editor/BlueprintGraph/Private/EdGraphSchema_K2.cpp
+++ b/Engine/Source/Editor/BlueprintGraph/Private/EdGraphSchema_K2.cpp
@@ -42,6 +42,8 @@ const FName FBlueprintMetadata::MD_NotAllowableBlueprintVariableType(TEXT("NotBl
 const FName FBlueprintMetadata::MD_BlueprintSpawnableComponent(TEXT("BlueprintSpawnableComponent"));
 const FName FBlueprintMetadata::MD_IsBlueprintBase(TEXT("IsBlueprintBase"));
 const FName FBlueprintMetadata::MD_RestrictedToClasses(TEXT("RestrictedToClasses"));
+const FName FBlueprintMetadata::MD_ChildCanTick(TEXT("ChildCanTick"));
+const FName FBlueprintMetadata::MD_ChildCannotTick(TEXT("ChildCannotTick"));
 
 const FName FBlueprintMetadata::MD_Protected(TEXT("BlueprintProtected"));
 const FName FBlueprintMetadata::MD_Latent(TEXT("Latent"));
diff --git a/Engine/Source/Editor/KismetCompiler/Private/KismetCompiler.cpp b/Engine/Source/Editor/KismetCompiler/Private/KismetCompiler.cpp
--- a/Engine/Source/Editor/KismetCompiler/Private/KismetCompiler.cpp
+++ b/Engine/Source/Editor/KismetCompiler/Private/KismetCompiler.cpp
@@ -3753,28 +3753,71 @@ void FKismetCompilerContext::SetCanEverTick() const
 	// RECEIVE TICK
 	if (!TickFunction->bCanEverTick)
 	{
+		// Make sure that both AActor and UActorComponent have the same name for their tick method
 		static FName ReceiveTickName(GET_FUNCTION_NAME_CHECKED(AActor, ReceiveTick));
-		static FName ComponentReceiveTickName(GET_FUNCTION_NAME_CHECKED(UActorComponent, ReceiveTick)); // Only doing this to ensure that both classes have the correct, same name
-		const UFunction* ReciveTickEvent = FKismetCompilerUtilities::FindOverriddenImplementableEvent(ReceiveTickName, NewClass);
-		if (ReciveTickEvent)
+		static FName ComponentReceiveTickName(GET_FUNCTION_NAME_CHECKED(UActorComponent, ReceiveTick));
+
+		if (const UFunction* ReceiveTickEvent = FKismetCompilerUtilities::FindOverriddenImplementableEvent(ReceiveTickName, NewClass))
 		{
-			static const FName ChildCanTickName = TEXT("ChildCanTick");
+			// We have a tick node, but are we allowed to?
+
+			const UEngine* EngineSettings = GetDefault<UEngine>();
+			const bool bAllowTickingByDefault = EngineSettings->bCanBlueprintsTickByDefault;
+
 			const UClass* FirstNativeClass = FBlueprintEditorUtils::FindFirstNativeClass(NewClass);
-			const bool bOverrideFlags = (AActor::StaticClass() == FirstNativeClass) || (UActorComponent::StaticClass() == FirstNativeClass) || (FirstNativeClass && FirstNativeClass->HasMetaData(ChildCanTickName));
-			if (bOverrideFlags)
+			const bool bHasCanTickMetadata = (FirstNativeClass != nullptr) && FirstNativeClass->HasMetaData(FBlueprintMetadata::MD_ChildCanTick);
+			const bool bHasCannotTickMetadata = (FirstNativeClass != nullptr) && FirstNativeClass->HasMetaData(FBlueprintMetadata::MD_ChildCannotTick);
+			const bool bHasUniversalParent = (FirstNativeClass != nullptr) && ((AActor::StaticClass() == FirstNativeClass) || (UActorComponent::StaticClass() == FirstNativeClass));
+
+			if (bHasCanTickMetadata && bHasCannotTickMetadata)
+			{
+				// User error: The C++ class has conflicting metadata
+				const FString ConlictingMetadataWarning = FString::Printf(
+					*LOCTEXT("HasBothCanAndCannotMetadata", "Native class %s has both '%s' and '%s' metadata specified, they are mutually exclusive and '%s' will win.").ToString(),
+					*FirstNativeClass->GetPathName(),
+					*FBlueprintMetadata::MD_ChildCanTick.ToString(),
+					*FBlueprintMetadata::MD_ChildCannotTick.ToString(),
+					*FBlueprintMetadata::MD_ChildCannotTick.ToString());
+				MessageLog.Warning(*ConlictingMetadataWarning);
+			}
+
+			if (bHasCannotTickMetadata)
 			{
-				TickFunction->bCanEverTick = true;
+				// This could only happen if someone adds bad metadata to AActor or UActorComponent directly
+				check(!bHasUniversalParent);
+
+				// Parent class has forbidden us to tick
+				const FString NativeClassSaidNo = FString::Printf(
+					*LOCTEXT("NativeClassProhibitsTicking", "@@ is not allowed as the C++ parent class %s has disallowed Blueprint subclasses from ticking.  Please consider using a Timer instead of Tick.").ToString(),
+					*FirstNativeClass->GetPathName(),
+					*FBlueprintMetadata::MD_ChildCannotTick.ToString());
+				MessageLog.Warning(*NativeClassSaidNo, FindLocalEntryPoint(ReceiveTickEvent));
 			}
-			else if (!TickFunction->bCanEverTick)
+			else
 			{
-				const FString ReceivTickEventWarning = FString::Printf( 
-					*LOCTEXT("ReceiveTick_CanNeverTick", "Blueprint %s has the ReceiveTick @@ event, but it can never tick.  Please consider using a Timer instead of a tick or you can enable Tick using code in one of the following ways: set ChildCanTick in the metadata on the parent class, or set bCanEverTick to true.").ToString(), *NewClass->GetName());
-				MessageLog.Warning( *ReceivTickEventWarning, FindLocalEntryPoint(ReciveTickEvent) );
+				if (bAllowTickingByDefault || bHasUniversalParent || bHasCanTickMetadata)
+				{
+					// We're allowed to tick for one reason or another
+					TickFunction->bCanEverTick = true;
+				}
+				else
+				{
+					// Nothing allowing us to tick
+					const FString ReceiveTickEventWarning = FString::Printf(
+						*LOCTEXT("ReceiveTick_CanNeverTick", "@@ is not allowed for Blueprints based on the C++ parent class %s, so it will never Tick!").ToString(),
+						*FirstNativeClass->GetPathName());
+					MessageLog.Warning(*ReceiveTickEventWarning, FindLocalEntryPoint(ReceiveTickEvent));
+
+					const FString ReceiveTickEventRemedies = FString::Printf(
+						*LOCTEXT("RecieveTick_CanNeverTickRemedies", "You can solve this in several ways:\n  1) Consider using a Timer instead of Tick.\n  2) Add meta=(%s) to the parent C++ class\n  3) Reparent the Blueprint to AActor or UActorComponent, which can always tick.").ToString(),
+						*FBlueprintMetadata::MD_ChildCanTick.ToString());
+					MessageLog.Warning(*ReceiveTickEventRemedies);
+				}
 			}
 		}
 	}
 
-	if(TickFunction->bCanEverTick != bOldFlag)
+	if (TickFunction->bCanEverTick != bOldFlag)
 	{
 		UE_LOG(LogK2Compiler, Verbose, TEXT("Overridden flag for class '%s': CanEverTick %s "), *NewClass->GetName(),
 			TickFunction->bCanEverTick ? *(GTrue.ToString()) : *(GFalse.ToString()) );
diff --git a/Engine/Source/Editor/UnrealEd/Private/Editor.cpp b/Engine/Source/Editor/UnrealEd/Private/Editor.cpp
--- a/Engine/Source/Editor/UnrealEd/Private/Editor.cpp
+++ b/Engine/Source/Editor/UnrealEd/Private/Editor.cpp
@@ -1841,12 +1841,23 @@ void UEditorEngine::PostEditChangeProperty(FPropertyChangedEvent& PropertyChange
 
 	const FName PropertyName = PropertyChangedEvent.Property ? PropertyChangedEvent.Property->GetFName() : NAME_None;
 
-	if( PropertyName == FName( TEXT( "MaximumLoopIterationCount" )))
+	if (PropertyName == GET_MEMBER_NAME_CHECKED(UEngine, MaximumLoopIterationCount))
 	{
 		// Clamp to a reasonable range and feed the new value to the script core
 		MaximumLoopIterationCount = FMath::Clamp( MaximumLoopIterationCount, 100, 10000000 );
 		FBlueprintCoreDelegates::SetScriptMaximumLoopIterations( MaximumLoopIterationCount );
 	}
+	else if (PropertyName == GET_MEMBER_NAME_CHECKED(UEngine, bCanBlueprintsTickByDefault))
+	{
+		FScopedSlowTask SlowTask(100, LOCTEXT("DirtyingBlueprintsDueToTickChange", "InvalidatingAllBlueprints"));
+
+		// Flag all Blueprints as out of date (this doesn't dirty the package as needs saving but will force a recompile during PIE)
+		for (TObjectIterator<UBlueprint> BlueprintIt; BlueprintIt; ++BlueprintIt)
+		{
+			UBlueprint* Blueprint = *BlueprintIt;
+			Blueprint->Status = BS_Dirty;
+		}
+	}
 }
 
 void UEditorEngine::Cleanse( bool ClearSelection, bool Redraw, const FText& TransReset )
diff --git a/Engine/Source/Runtime/CoreUObject/Public/UObject/ObjectBase.h b/Engine/Source/Runtime/CoreUObject/Public/UObject/ObjectBase.h
--- a/Engine/Source/Runtime/CoreUObject/Public/UObject/ObjectBase.h
+++ b/Engine/Source/Runtime/CoreUObject/Public/UObject/ObjectBase.h
@@ -892,9 +892,12 @@ namespace UM
 		/// [ClassMetadata] Used for Actor Component classes. If present indicates that it can be spawned by a Blueprint.
 		BlueprintSpawnableComponent,
 
-		/// [ClassMetadata] Used for Actor classes. If the native class cannot tick, Blueprint generated classes based this Actor can have bCanEverTick flag overridden.
+		/// [ClassMetadata] Used for Actor and Component classes. If the native class cannot tick, Blueprint generated classes based this Actor or Component can have bCanEverTick flag overridden even if bCanBlueprintsTickByDefault is false.
 		ChildCanTick,
 
+		/// [ClassMetadata] Used for Actor and Component classes. If the native class cannot tick, Blueprint generated classes based this Actor or Component can never tick even if bCanBlueprintsTickByDefault is true.
+		ChildCannotTick,
+
 		/// [ClassMetadata] For BehaviorTree nodes indicates that the class is deprecated and will display a warning when compiled.
 		DeprecatedNode,
 
diff --git a/Engine/Source/Runtime/Engine/Classes/Engine/Engine.h b/Engine/Source/Runtime/Engine/Classes/Engine/Engine.h
--- a/Engine/Source/Runtime/Engine/Classes/Engine/Engine.h
+++ b/Engine/Source/Runtime/Engine/Classes/Engine/Engine.h
@@ -1171,6 +1171,14 @@ public:
 	UPROPERTY(EditAnywhere, config, Category=Blueprints)
 	int32 MaximumLoopIterationCount;
 
+	// Controls whether Blueprint subclasses of actors or components can tick by default.
+	//
+	// Blueprints that derive from native C++ classes that have bCanEverTick=true will always be able to tick
+	// Blueprints that derive from exactly AActor or UActorComponent will always be able to tick
+	// Otherwise, they can tick as long as the parent doesn't have meta=(ChildCannotTick) and either bCanBlueprintsTickByDefault is true or the parent has meta=(ChildCanTick)
+	UPROPERTY(EditAnywhere, config, Category=Blueprints)
+	uint32 bCanBlueprintsTickByDefault:1;
+
 	/** @todo document */
 	UPROPERTY(config)
 	uint32 bEnableEditorPSysRealtimeLOD:1;
diff --git a/Engine/Source/Runtime/Engine/Private/GameEngine.cpp b/Engine/Source/Runtime/Engine/Private/GameEngine.cpp
--- a/Engine/Source/Runtime/Engine/Private/GameEngine.cpp
+++ b/Engine/Source/Runtime/Engine/Private/GameEngine.cpp
@@ -387,6 +387,8 @@ UEngine::UEngine(const FObjectInitializer& ObjectInitializer)
 
 	BeginStreamingPauseDelegate = NULL;
 	EndStreamingPauseDelegate = NULL;
+
+	bCanBlueprintsTickByDefault = true;
 }
 
 void UGameEngine::Init(IEngineLoop* InEngineLoop)
