[INTEGRATE] Change 2441047 by [REDACTED_PERSON][REDACTED_HANDLE] on 2015/02/11 00:38:29

	Blueprints: Add support for class metadata IgnoreCategoryKeywordsInSubclasses, which prevents inheritance of showCategories, hideCategories, AutoExpandCategories, and AutoCollapseCategories metadata when placed on the immediate parent of a Blueprint

	UnrealHeaderTool: Add a new UCLASS() specifier ComponentWrapperClass, which implies IgnoreCategoryKeywordsInSubclasses metadata

	Engine: Add ComponentWrapperClass to all actors that contain a component marked with ExposeFunctionCategories
	[UE-9115]

diff --git a/Engine/Plugins/2D/Paper2D/Source/Paper2D/Classes/PaperFlipbookActor.h b/Engine/Plugins/2D/Paper2D/Source/Paper2D/Classes/PaperFlipbookActor.h
--- a/Engine/Plugins/2D/Paper2D/Source/Paper2D/Classes/PaperFlipbookActor.h
+++ b/Engine/Plugins/2D/Paper2D/Source/Paper2D/Classes/PaperFlipbookActor.h
@@ -11,7 +11,7 @@
  * This actor is created when you drag a flipbook asset from the content browser into the level, and
  * it is just a thin wrapper around a UPaperFlipbookComponent that actually references the asset.
  */
-UCLASS()
+UCLASS(ComponentWrapperClass)
 class PAPER2D_API APaperFlipbookActor : public AActor
 {
 	GENERATED_UCLASS_BODY()
diff --git a/Engine/Plugins/2D/Paper2D/Source/Paper2D/Classes/PaperSpriteActor.h b/Engine/Plugins/2D/Paper2D/Source/Paper2D/Classes/PaperSpriteActor.h
--- a/Engine/Plugins/2D/Paper2D/Source/Paper2D/Classes/PaperSpriteActor.h
+++ b/Engine/Plugins/2D/Paper2D/Source/Paper2D/Classes/PaperSpriteActor.h
@@ -11,7 +11,7 @@
  * This actor is created when you drag a sprite asset from the content browser into the level, and
  * it is just a thin wrapper around a UPaperSpriteComponent that actually references the asset.
  */
-UCLASS(meta=(ChildCanTick))
+UCLASS(ComponentWrapperClass, meta = (ChildCanTick))
 class PAPER2D_API APaperSpriteActor : public AActor
 {
 	GENERATED_UCLASS_BODY()
diff --git a/Engine/Plugins/2D/Paper2D/Source/Paper2D/Classes/PaperTileMapActor.h b/Engine/Plugins/2D/Paper2D/Source/Paper2D/Classes/PaperTileMapActor.h
--- a/Engine/Plugins/2D/Paper2D/Source/Paper2D/Classes/PaperTileMapActor.h
+++ b/Engine/Plugins/2D/Paper2D/Source/Paper2D/Classes/PaperTileMapActor.h
@@ -12,7 +12,7 @@
  * This actor is created when you drag a tile map asset from the content browser into the level, and
  * it is just a thin wrapper around a UPaperTileMapComponent that actually references the asset.
  */
-UCLASS()
+UCLASS(ComponentWrapperClass)
 class PAPER2D_API APaperTileMapActor : public AActor
 {
 	GENERATED_UCLASS_BODY()
diff --git a/Engine/Source/Editor/BlueprintGraph/Classes/EdGraphSchema_K2.h b/Engine/Source/Editor/BlueprintGraph/Classes/EdGraphSchema_K2.h
--- a/Engine/Source/Editor/BlueprintGraph/Classes/EdGraphSchema_K2.h
+++ b/Engine/Source/Editor/BlueprintGraph/Classes/EdGraphSchema_K2.h
@@ -65,6 +65,9 @@ public:
 	/// [ClassMetadata] Used for Actor and Component classes. If the native class cannot tick, Blueprint generated classes based this Actor or Component can never tick even if bCanBlueprintsTickByDefault is true.
 	static const FName MD_ChildCannotTick;
 
+	/// [ClassMetadata] Used to make the first subclass of a class ignore all inherited showCategories and hideCategories commands
+	static const FName MD_IgnoreCategoryKeywordsInSubclasses;
+
 	//    function metadata
 	/** Specifies a UFUNCTION as Kismet protected, which can only be called from itself */
 	static const FName MD_Protected;
diff --git a/Engine/Source/Editor/BlueprintGraph/Private/EdGraphSchema_K2.cpp b/Engine/Source/Editor/BlueprintGraph/Private/EdGraphSchema_K2.cpp
--- a/Engine/Source/Editor/BlueprintGraph/Private/EdGraphSchema_K2.cpp
+++ b/Engine/Source/Editor/BlueprintGraph/Private/EdGraphSchema_K2.cpp
@@ -44,6 +44,8 @@ const FName FBlueprintMetadata::MD_IsBlueprintBase(TEXT("IsBlueprintBase"));
 const FName FBlueprintMetadata::MD_RestrictedToClasses(TEXT("RestrictedToClasses"));
 const FName FBlueprintMetadata::MD_ChildCanTick(TEXT("ChildCanTick"));
 const FName FBlueprintMetadata::MD_ChildCannotTick(TEXT("ChildCannotTick"));
+const FName FBlueprintMetadata::MD_IgnoreCategoryKeywordsInSubclasses(TEXT("IgnoreCategoryKeywordsInSubclasses"));
+
 
 const FName FBlueprintMetadata::MD_Protected(TEXT("BlueprintProtected"));
 const FName FBlueprintMetadata::MD_Latent(TEXT("Latent"));
diff --git a/Engine/Source/Editor/KismetCompiler/Private/KismetCompiler.cpp b/Engine/Source/Editor/KismetCompiler/Private/KismetCompiler.cpp
--- a/Engine/Source/Editor/KismetCompiler/Private/KismetCompiler.cpp
+++ b/Engine/Source/Editor/KismetCompiler/Private/KismetCompiler.cpp
@@ -1700,23 +1700,27 @@ void FKismetCompilerContext::FinishCompilingClass(UClass* Class)
 
 		// Copy the category info from the parent class
 #if WITH_EDITORONLY_DATA
-		FEditorCategoryUtils::GetClassHideCategories(ParentClass, AllHideCategories);
-		if (ParentClass->HasMetaData(TEXT("ShowCategories")))
+		if (!ParentClass->HasMetaData(FBlueprintMetadata::MD_IgnoreCategoryKeywordsInSubclasses))
 		{
-			Class->SetMetaData(TEXT("ShowCategories"), *ParentClass->GetMetaData("ShowCategories"));
+			FEditorCategoryUtils::GetClassHideCategories(ParentClass, AllHideCategories);
+			if (ParentClass->HasMetaData(TEXT("ShowCategories")))
+			{
+				Class->SetMetaData(TEXT("ShowCategories"), *ParentClass->GetMetaData("ShowCategories"));
+			}
+			if (ParentClass->HasMetaData(TEXT("AutoExpandCategories")))
+			{
+				Class->SetMetaData(TEXT("AutoExpandCategories"), *ParentClass->GetMetaData("AutoExpandCategories"));
+			}
+			if (ParentClass->HasMetaData(TEXT("AutoCollapseCategories")))
+			{
+				Class->SetMetaData(TEXT("AutoCollapseCategories"), *ParentClass->GetMetaData("AutoCollapseCategories"));
+			}
 		}
+
 		if (ParentClass->HasMetaData(TEXT("HideFunctions")))
 		{
 			Class->SetMetaData(TEXT("HideFunctions"), *ParentClass->GetMetaData("HideFunctions"));
 		}
-		if (ParentClass->HasMetaData(TEXT("AutoExpandCategories")))
-		{
-			Class->SetMetaData(TEXT("AutoExpandCategories"), *ParentClass->GetMetaData("AutoExpandCategories"));
-		}
-		if (ParentClass->HasMetaData(TEXT("AutoCollapseCategories")))
-		{
-			Class->SetMetaData(TEXT("AutoCollapseCategories"), *ParentClass->GetMetaData("AutoCollapseCategories"));
-		}
 		
 		// Blueprinted Components are always Blueprint Spawnable
 		if (ParentClass->IsChildOf(UActorComponent::StaticClass()))
diff --git a/Engine/Source/Programs/UnrealHeaderTool/Private/HeaderParser.cpp b/Engine/Source/Programs/UnrealHeaderTool/Private/HeaderParser.cpp
--- a/Engine/Source/Programs/UnrealHeaderTool/Private/HeaderParser.cpp
+++ b/Engine/Source/Programs/UnrealHeaderTool/Private/HeaderParser.cpp
@@ -4194,6 +4194,10 @@ void FHeaderParser::CompileClassDeclaration(FClasses& AllClasses)
 		{
 			Class->ClassFlags |= CLASS_Intrinsic;
 		}
+		else if (Specifier == TEXT("ComponentWrapperClass"))
+		{
+			MetaData.Add(TEXT("IgnoreCategoryKeywordsInSubclasses"), TEXT("true"));
+		}
 		else if (Specifier == TEXT("within"))
 		{
 			FString WithinNameStr = RequireExactlyOneSpecifierValue(PropSpecifier);
diff --git a/Engine/Source/Runtime/CoreUObject/Public/UObject/ObjectBase.h b/Engine/Source/Runtime/CoreUObject/Public/UObject/ObjectBase.h
--- a/Engine/Source/Runtime/CoreUObject/Public/UObject/ObjectBase.h
+++ b/Engine/Source/Runtime/CoreUObject/Public/UObject/ObjectBase.h
@@ -615,6 +615,8 @@ namespace UC
 		showCategories,
 		/// Hides the specified categories in a property viewer. Usage: hideCategories=CategoryName or hideCategories=(category0, category1, ...)
 		hideCategories,
+		/// Indicates that this class is a wrapper class for a component with little intrinsic functionality (this causes things like hideCategories and showCategories to be ignored if the class is subclassed in a Blueprint)
+		ComponentWrapperClass,
 		/// Shows the specified function in a property viewer. Usage: showFunctions=FunctionName or showFunctions=(category0, category1, ...)
 		showFunctions,
 		/// Hides the specified function in a property viewer. Usage: hideFunctions=FunctionName or hideFunctions=(category0, category1, ...)
@@ -898,6 +900,9 @@ namespace UM
 		/// [ClassMetadata] Used for Actor and Component classes. If the native class cannot tick, Blueprint generated classes based this Actor or Component can never tick even if bCanBlueprintsTickByDefault is true.
 		ChildCannotTick,
 
+		/// [ClassMetadata] Used to make the first subclass of a class ignore all inherited showCategories and hideCategories commands
+		IgnoreCategoryKeywordsInSubclasses,
+
 		/// [ClassMetadata] For BehaviorTree nodes indicates that the class is deprecated and will display a warning when compiled.
 		DeprecatedNode,
 
diff --git a/Engine/Source/Runtime/Engine/Classes/Animation/SkeletalMeshActor.h b/Engine/Source/Runtime/Engine/Classes/Animation/SkeletalMeshActor.h
--- a/Engine/Source/Runtime/Engine/Classes/Animation/SkeletalMeshActor.h
+++ b/Engine/Source/Runtime/Engine/Classes/Animation/SkeletalMeshActor.h
@@ -12,7 +12,7 @@
  * @see https://docs.unrealengine.com/latest/INT/Engine/Content/Types/SkeletalMeshes/
  * @see USkeletalMesh
  */
-UCLASS(ClassGroup=ISkeletalMeshes, Blueprintable, ConversionRoot, meta=(ChildCanTick))
+UCLASS(ClassGroup=ISkeletalMeshes, Blueprintable, ComponentWrapperClass, ConversionRoot, meta=(ChildCanTick))
 class ENGINE_API ASkeletalMeshActor : public AActor, public IMatineeAnimInterface
 {
 	GENERATED_UCLASS_BODY()
diff --git a/Engine/Source/Runtime/Engine/Classes/Components/SceneComponent.h b/Engine/Source/Runtime/Engine/Classes/Components/SceneComponent.h
--- a/Engine/Source/Runtime/Engine/Classes/Components/SceneComponent.h
+++ b/Engine/Source/Runtime/Engine/Classes/Components/SceneComponent.h
@@ -77,7 +77,7 @@ DECLARE_DYNAMIC_MULTICAST_DELEGATE_OneParam(FPhysicsVolumeChanged, class APhysic
  * Useful as a 'dummy' component in the hierarchy to offset others.
  * @see [Scene Components](https://docs.unrealengine.com/latest/INT/Programming/UnrealArchitecture/Actors/Components/index.html#scenecomponents)
  */
-UCLASS(ClassGroup=(Utility, Common), BlueprintType, hideCategories=(Trigger, PhysicsVolume), meta=(BlueprintSpawnableComponent, ShortTooltip="A Scene Component is a component that has a scene transform and can be attached to other scene components."))
+UCLASS(ClassGroup=(Utility, Common), BlueprintType, hideCategories=(Trigger, PhysicsVolume), meta=(BlueprintSpawnableComponent, IgnoreCategoryKeywordsInSubclasses, ShortTooltip="A Scene Component is a component that has a scene transform and can be attached to other scene components."))
 class ENGINE_API USceneComponent : public UActorComponent
 {
 	GENERATED_BODY()
diff --git a/Engine/Source/Runtime/Engine/Classes/Engine/DecalActor.h b/Engine/Source/Runtime/Engine/Classes/Engine/DecalActor.h
--- a/Engine/Source/Runtime/Engine/Classes/Engine/DecalActor.h
+++ b/Engine/Source/Runtime/Engine/Classes/Engine/DecalActor.h
@@ -11,7 +11,7 @@ class UBoxComponent;
 /**
 * Decal to add texture/material on existing geometry.
 */
-UCLASS(hidecategories=(Collision, Attachment, Actor, Input, Replication), showcategories=("Input|MouseInput", "Input|TouchInput"), MinimalAPI)
+UCLASS(hideCategories=(Collision, Attachment, Actor, Input, Replication), showCategories=("Input|MouseInput", "Input|TouchInput"), ComponentWrapperClass,MinimalAPI)
 class ADecalActor
 	: public AActor
 {
diff --git a/Engine/Source/Runtime/Engine/Classes/Engine/Light.h b/Engine/Source/Runtime/Engine/Classes/Engine/Light.h
--- a/Engine/Source/Runtime/Engine/Classes/Engine/Light.h
+++ b/Engine/Source/Runtime/Engine/Classes/Engine/Light.h
@@ -5,7 +5,7 @@
 #include "Components/LightComponent.h"
 #include "Light.generated.h"
 
-UCLASS(Abstract, ClassGroup=Lights, hidecategories=(Input,Collision,Replication), showcategories=("Input|MouseInput", "Input|TouchInput"), ConversionRoot, meta=(ChildCanTick))
+UCLASS(Abstract, ClassGroup=Lights, hideCategories=(Input,Collision,Replication), showCategories=("Input|MouseInput", "Input|TouchInput"), ComponentWrapperClass, ConversionRoot, meta=(ChildCanTick))
 class ENGINE_API ALight : public AActor
 {
 	GENERATED_UCLASS_BODY()
diff --git a/Engine/Source/Runtime/Engine/Classes/Engine/PointLight.h b/Engine/Source/Runtime/Engine/Classes/Engine/PointLight.h
--- a/Engine/Source/Runtime/Engine/Classes/Engine/PointLight.h
+++ b/Engine/Source/Runtime/Engine/Classes/Engine/PointLight.h
@@ -4,7 +4,7 @@
 #include "Light.h"
 #include "PointLight.generated.h"
 
-UCLASS(ClassGroup=(Lights, PointLights),MinimalAPI, meta=(ChildCanTick))
+UCLASS(ClassGroup=(Lights, PointLights), ComponentWrapperClass, MinimalAPI, meta=(ChildCanTick))
 class APointLight : public ALight
 {
 	GENERATED_UCLASS_BODY()
diff --git a/Engine/Source/Runtime/Engine/Classes/Engine/SkyLight.h b/Engine/Source/Runtime/Engine/Classes/Engine/SkyLight.h
--- a/Engine/Source/Runtime/Engine/Classes/Engine/SkyLight.h
+++ b/Engine/Source/Runtime/Engine/Classes/Engine/SkyLight.h
@@ -4,7 +4,7 @@
 #include "GameFramework/Info.h"
 #include "SkyLight.generated.h"
 
-UCLASS(ClassGroup=Lights, hidecategories=(Input,Collision,Replication,Info), showcategories=("Input|MouseInput", "Input|TouchInput"), ConversionRoot, Blueprintable)
+UCLASS(ClassGroup=Lights, hidecategories=(Input,Collision,Replication,Info), showcategories=("Input|MouseInput", "Input|TouchInput"), ComponentWrapperClass, ConversionRoot, Blueprintable)
 class ENGINE_API ASkyLight : public AInfo
 {
 	GENERATED_UCLASS_BODY()
diff --git a/Engine/Source/Runtime/Engine/Classes/Engine/SplineMeshActor.h b/Engine/Source/Runtime/Engine/Classes/Engine/SplineMeshActor.h
--- a/Engine/Source/Runtime/Engine/Classes/Engine/SplineMeshActor.h
+++ b/Engine/Source/Runtime/Engine/Classes/Engine/SplineMeshActor.h
@@ -8,7 +8,7 @@
  *
  * @see USplineMeshComponent
  */
-UCLASS(hidecategories = (Input), showcategories = ("Input|MouseInput", "Input|TouchInput"), ConversionRoot, meta = (ChildCanTick))
+UCLASS(hideCategories = (Input), showCategories = ("Input|MouseInput", "Input|TouchInput"), ConversionRoot, ComponentWrapperClass, meta = (ChildCanTick))
 class ENGINE_API ASplineMeshActor : public AActor
 {
 	GENERATED_UCLASS_BODY()
diff --git a/Engine/Source/Runtime/Engine/Classes/Engine/SpotLight.h b/Engine/Source/Runtime/Engine/Classes/Engine/SpotLight.h
--- a/Engine/Source/Runtime/Engine/Classes/Engine/SpotLight.h
+++ b/Engine/Source/Runtime/Engine/Classes/Engine/SpotLight.h
@@ -5,7 +5,7 @@
 #include "Engine/Light.h"
 #include "SpotLight.generated.h"
 
-UCLASS(ClassGroup=(Lights, SpotLights), MinimalAPI, meta=(ChildCanTick))
+UCLASS(ClassGroup=(Lights, SpotLights), MinimalAPI, ComponentWrapperClass, meta=(ChildCanTick))
 class ASpotLight : public ALight
 {
 	GENERATED_UCLASS_BODY()
diff --git a/Engine/Source/Runtime/Engine/Classes/Engine/StaticMeshActor.h b/Engine/Source/Runtime/Engine/Classes/Engine/StaticMeshActor.h
--- a/Engine/Source/Runtime/Engine/Classes/Engine/StaticMeshActor.h
+++ b/Engine/Source/Runtime/Engine/Classes/Engine/StaticMeshActor.h
@@ -11,7 +11,7 @@
  * @see https://docs.unrealengine.com/latest/INT/Engine/Actors/StaticMeshActor/
  * @see UStaticMesh
  */
-UCLASS(hidecategories=(Input), showcategories=("Input|MouseInput", "Input|TouchInput"), ConversionRoot, meta=(ChildCanTick))
+UCLASS(hidecategories=(Input), showcategories=("Input|MouseInput", "Input|TouchInput"), ConversionRoot, ComponentWrapperClass, meta=(ChildCanTick))
 class ENGINE_API AStaticMeshActor : public AActor
 {
 	GENERATED_UCLASS_BODY()
diff --git a/Engine/Source/Runtime/Engine/Classes/Engine/TextRenderActor.h b/Engine/Source/Runtime/Engine/Classes/Engine/TextRenderActor.h
--- a/Engine/Source/Runtime/Engine/Classes/Engine/TextRenderActor.h
+++ b/Engine/Source/Runtime/Engine/Classes/Engine/TextRenderActor.h
@@ -8,7 +8,7 @@
 
 #include "TextRenderActor.generated.h"
 
-UCLASS(MinimalAPI, hidecategories=(Collision, Attachment, Actor))
+UCLASS(MinimalAPI, ComponentWrapperClass, hideCategories = (Collision, Attachment, Actor))
 class ATextRenderActor : public AActor
 {
 	GENERATED_UCLASS_BODY()
diff --git a/Engine/Source/Runtime/Engine/Classes/Particles/Emitter.h b/Engine/Source/Runtime/Engine/Classes/Particles/Emitter.h
--- a/Engine/Source/Runtime/Engine/Classes/Particles/Emitter.h
+++ b/Engine/Source/Runtime/Engine/Classes/Particles/Emitter.h
@@ -44,7 +44,7 @@ DECLARE_DYNAMIC_MULTICAST_DELEGATE_SixParams( FParticleDeathSignature, FName, Ev
 	*/
 DECLARE_DYNAMIC_MULTICAST_DELEGATE_EightParams( FParticleCollisionSignature, FName, EventName, float, EmitterTime, int32, ParticleTime, FVector, Location, FVector, Velocity, FVector, Direction, FVector, Normal, FName, BoneName);
 
-UCLASS(hidecategories=(Activation,"Components|Activation",Input,Collision,"Game|Damage"))
+UCLASS(hideCategories=(Activation,"Components|Activation",Input,Collision,"Game|Damage"), ComponentWrapperClass)
 class ENGINE_API AEmitter : public AActor
 {
 	GENERATED_UCLASS_BODY()
diff --git a/Engine/Source/Runtime/Engine/Classes/PhysicsEngine/DestructibleActor.h b/Engine/Source/Runtime/Engine/Classes/PhysicsEngine/DestructibleActor.h
--- a/Engine/Source/Runtime/Engine/Classes/PhysicsEngine/DestructibleActor.h
+++ b/Engine/Source/Runtime/Engine/Classes/PhysicsEngine/DestructibleActor.h
@@ -11,7 +11,7 @@ class UDestructibleComponent;
 /** Delegate for notification when fracture occurs */
 DECLARE_DYNAMIC_MULTICAST_DELEGATE_TwoParams(FActorFractureSignature, const FVector &, HitPoint, const FVector &, HitDirection);
 
-UCLASS(MinimalAPI, hidecategories=(Input), showcategories=("Input|MouseInput", "Input|TouchInput"))
+UCLASS(MinimalAPI, hideCategories=(Input), showCategories=("Input|MouseInput", "Input|TouchInput"), ComponentWrapperClass)
 class ADestructibleActor : public AActor
 {
 	GENERATED_UCLASS_BODY()
diff --git a/Engine/Source/Runtime/Engine/Classes/PhysicsEngine/PhysicsConstraintActor.h b/Engine/Source/Runtime/Engine/Classes/PhysicsEngine/PhysicsConstraintActor.h
--- a/Engine/Source/Runtime/Engine/Classes/PhysicsEngine/PhysicsConstraintActor.h
+++ b/Engine/Source/Runtime/Engine/Classes/PhysicsEngine/PhysicsConstraintActor.h
@@ -7,7 +7,7 @@
 #include "PhysicsEngine/RigidBodyBase.h"
 #include "PhysicsConstraintActor.generated.h"
 
-UCLASS(ConversionRoot, MinimalAPI)
+UCLASS(ConversionRoot, MinimalAPI, ComponentWrapperClass)
 class APhysicsConstraintActor : public ARigidBodyBase
 {
 	GENERATED_UCLASS_BODY()
diff --git a/Engine/Source/Runtime/Engine/Classes/PhysicsEngine/PhysicsThruster.h b/Engine/Source/Runtime/Engine/Classes/PhysicsEngine/PhysicsThruster.h
--- a/Engine/Source/Runtime/Engine/Classes/PhysicsEngine/PhysicsThruster.h
+++ b/Engine/Source/Runtime/Engine/Classes/PhysicsEngine/PhysicsThruster.h
@@ -10,7 +10,7 @@
  *	Attach one of these on an object using physics simulation and it will apply a force down the negative-X direction
  *	ie. point X in the direction you want the thrust in.
  */
-UCLASS(hidecategories=(Input,Collision,Replication), showcategories=("Input|MouseInput", "Input|TouchInput"))
+UCLASS(hideCategories=(Input,Collision,Replication), showCategories=("Input|MouseInput", "Input|TouchInput"), ComponentWrapperClass)
 class APhysicsThruster : public ARigidBodyBase
 {
 	GENERATED_UCLASS_BODY()
diff --git a/Engine/Source/Runtime/Engine/Classes/PhysicsEngine/RadialForceActor.h b/Engine/Source/Runtime/Engine/Classes/PhysicsEngine/RadialForceActor.h
--- a/Engine/Source/Runtime/Engine/Classes/PhysicsEngine/RadialForceActor.h
+++ b/Engine/Source/Runtime/Engine/Classes/PhysicsEngine/RadialForceActor.h
@@ -5,7 +5,7 @@
 #include "RigidBodyBase.h"
 #include "RadialForceActor.generated.h"
 
-UCLASS(MinimalAPI, hidecategories=(Collision, Input), showcategories=("Input|MouseInput", "Input|TouchInput"))
+UCLASS(MinimalAPI, hideCategories=(Collision, Input), showCategories=("Input|MouseInput", "Input|TouchInput"), ComponentWrapperClass)
 class ARadialForceActor : public ARigidBodyBase
 {
 	GENERATED_UCLASS_BODY()
diff --git a/Engine/Source/Runtime/Engine/Classes/Sound/AmbientSound.h b/Engine/Source/Runtime/Engine/Classes/Sound/AmbientSound.h
--- a/Engine/Source/Runtime/Engine/Classes/Sound/AmbientSound.h
+++ b/Engine/Source/Runtime/Engine/Classes/Sound/AmbientSound.h
@@ -4,7 +4,7 @@
 #include "AmbientSound.generated.h"
 
 /** A sound actor that can be placed in a level */
-UCLASS(AutoExpandCategories=Audio, ClassGroup=Sounds, MinimalAPI, hidecategories(Collision, Input, Game), showcategories=("Input|MouseInput", "Input|TouchInput", "Game|Damage"))
+UCLASS(AutoExpandCategories=Audio, ClassGroup=Sounds, MinimalAPI, hideCategories(Collision, Input, Game), showCategories=("Input|MouseInput", "Input|TouchInput", "Game|Damage"), ComponentWrapperClass)
 class AAmbientSound : public AActor
 {
 	GENERATED_UCLASS_BODY()
