[INTEGRATE] Change 2442890 by [REDACTED_PERSON][REDACTED_HANDLE] on 2015/02/11 23:14:08

	Blueprints: Prevent Blueprint UActorComponent subclasses from trying to use names that could conflict with autogenerated component names
	- Added _C stripping in the rename test code, since it is already present in the generated name code
	- Generalized the last-ditch-effort code that was looking only for properties to  consider all objects, so any future regressions would rename generated properties out of the way and create a BP compile error rather than letting StaticAllocateObject assert
	[UE-9009]

diff --git a/Engine/Source/Editor/KismetCompiler/Private/KismetCompilerMisc.cpp b/Engine/Source/Editor/KismetCompiler/Private/KismetCompilerMisc.cpp
--- a/Engine/Source/Editor/KismetCompiler/Private/KismetCompilerMisc.cpp
+++ b/Engine/Source/Editor/KismetCompiler/Private/KismetCompilerMisc.cpp
@@ -651,13 +651,12 @@ UProperty* FKismetCompilerUtilities::CreatePropertyOnScope(UStruct* Scope, const
 
 	FName ValidatedPropertyName = PropertyName;
 
-	// Check to see if there's already a property on this scope, and throw an internal compiler error if so
+	// Check to see if there's already a object on this scope with the same name, and throw an internal compiler error if so
 	// If this happens, it breaks the property link, which causes stack corruption and hard-to-track errors, so better to fail at this point
 	{
-		UProperty* ExistingProperty = FindObject<UProperty>(Scope, *PropertyName.ToString(), false);
-		if( ExistingProperty )
+		if (UObject* ExistingObject = FindObject<UObject>(Scope, *PropertyName.ToString(), false))
 		{
-			MessageLog.Error(*FString::Printf(TEXT("Internal Compiler Error:  Duplicate property %s on scope %s"), *PropertyName.ToString(), (Scope ? *Scope->GetName() : TEXT("None"))));
+			MessageLog.Error(*FString::Printf(TEXT("Internal Compiler Error:  Tried to create a property %s in scope %s, but another object of type %s already already exists there."), *PropertyName.ToString(), (Scope ? *Scope->GetName() : TEXT("None"))), *ExistingObject->GetFullName(Scope));
 
 			// Find a free name, so we can still create the property to make it easier to spot the duplicates, and avoid crashing
 			uint32 Counter = 0;
@@ -665,7 +664,7 @@ UProperty* FKismetCompilerUtilities::CreatePropertyOnScope(UStruct* Scope, const
 			do 
 			{
 				TestNameString = PropertyName.ToString() + FString::Printf(TEXT("_ERROR_DUPLICATE_%d"), Counter++);
-			} while (FindObject<UProperty>(Scope, *TestNameString, false) != NULL);
+			} while (FindObject<UObject>(Scope, *TestNameString, false) != NULL);
 
 			ValidatedPropertyName = FName(*TestNameString);
 		}
diff --git a/Engine/Source/Editor/UnrealEd/Private/Kismet2/ComponentEditorUtils.cpp b/Engine/Source/Editor/UnrealEd/Private/Kismet2/ComponentEditorUtils.cpp
--- a/Engine/Source/Editor/UnrealEd/Private/Kismet2/ComponentEditorUtils.cpp
+++ b/Engine/Source/Editor/UnrealEd/Private/Kismet2/ComponentEditorUtils.cpp
@@ -118,7 +118,8 @@ bool FComponentEditorUtils::IsValidVariableNameString(const UActorComponent* InC
 	if(bIsValid && InComponent != NULL)
 	{
 		// Next test to make sure the string doesn't conflict with the format that MakeUniqueObjectName() generates
-		FString MakeUniqueObjectNamePrefix = FString::Printf(TEXT("%s_"), *InComponent->GetClass()->GetName());
+		const FString ClassNameThatWillBeUsedInGenerator = FBlueprintEditorUtils::GetClassNameWithoutSuffix(InComponent->GetClass());
+		const FString MakeUniqueObjectNamePrefix = FString::Printf(TEXT("%s_"), *ClassNameThatWillBeUsedInGenerator);
 		if(InString.StartsWith(MakeUniqueObjectNamePrefix))
 		{
 			bIsValid = !InString.Replace(*MakeUniqueObjectNamePrefix, TEXT("")).IsNumeric();
