Removing custom events and other graph spawning nodes will free up the name.

KismetCompiler will clean up old EventGraphs generated during compilation, removing the Blueprint from being the outer.

diff --git a/Engine/Source/Editor/BlueprintGraph/Classes/K2Node_CustomEvent.h b/Engine/Source/Editor/BlueprintGraph/Classes/K2Node_CustomEvent.h
--- a/Engine/Source/Editor/BlueprintGraph/Classes/K2Node_CustomEvent.h
+++ b/Engine/Source/Editor/BlueprintGraph/Classes/K2Node_CustomEvent.h
@@ -32,6 +32,7 @@ class UK2Node_CustomEvent : public UK2Node_Event
 	// Begin UK2Node interface
 	BLUEPRINTGRAPH_API virtual void ValidateNodeDuringCompilation(class FCompilerResultsLog& MessageLog) const override;
 	virtual void GetMenuActions(FBlueprintActionDatabaseRegistrar& ActionRegistrar) const override;
+	BLUEPRINTGRAPH_API virtual bool NodeCausesStructuralBlueprintChange() const override { return true; }
 	// End UK2Node interface
 
 	// Begin UK2Node_EditablePinBase interface
diff --git a/Engine/Source/Editor/KismetCompiler/Private/KismetCompiler.cpp b/Engine/Source/Editor/KismetCompiler/Private/KismetCompiler.cpp
--- a/Engine/Source/Editor/KismetCompiler/Private/KismetCompiler.cpp
+++ b/Engine/Source/Editor/KismetCompiler/Private/KismetCompiler.cpp
@@ -2724,6 +2724,14 @@ void FKismetCompilerContext::CreateAndProcessUbergraph()
 			UbergraphContext->MarkAsInternalOrCppUseOnly();
 			UbergraphContext->SetExternalNetNameMap(&ClassScopeNetNameMap);
 
+			// We need to stop the old EventGraphs from having the Blueprint as an outer, it impacts renaming.
+			if(!Blueprint->HasAnyFlags(RF_NeedLoad|RF_NeedPostLoad))
+			{
+				for(UEdGraph* OldEventGraph : Blueprint->EventGraphs)
+				{
+					OldEventGraph->Rename(NULL, GetTransientPackage());
+				}
+			}
 			Blueprint->EventGraphs.Empty();
 
 			// Validate all the nodes in the graph
