UE-9747 Custom structure in blueprint corrupts Project

Kismet COmpiler Error, when stub event graph has the same name as an existing graph (it was a crash before)

diff --git a/Engine/Source/Editor/KismetCompiler/Private/KismetCompiler.cpp b/Engine/Source/Editor/KismetCompiler/Private/KismetCompiler.cpp
--- a/Engine/Source/Editor/KismetCompiler/Private/KismetCompiler.cpp
+++ b/Engine/Source/Editor/KismetCompiler/Private/KismetCompiler.cpp
@@ -2355,6 +2355,17 @@ void FKismetCompilerContext::CreateFunctionStubForEvent(UK2Node_Event* SrcEventN
 
 	// Create the stub graph and add it to the list of functions to compile
 
+	auto ExistingGraph = static_cast<UObject*>(FindObjectWithOuter(OwnerOfTemporaries, UEdGraph::StaticClass(), EventNodeName));
+	if (ExistingGraph && !ExistingGraph->HasAnyFlags(RF_Transient))
+	{
+		// There is no match, so the function parameters must have changed.  Throw an error, and force them to refresh
+		MessageLog.Error(
+			*FString::Printf(*LOCTEXT("CannotCreateStubForEvent_Error", "Graph named '%s' already exists in '%s'. Another one cannot be generated from @@").ToString()
+			, *EventNodeName.ToString()
+			, *GetNameSafe(OwnerOfTemporaries))
+			, SrcEventNode);
+		return;
+	}
 	UEdGraph* ChildStubGraph = NewObject<UEdGraph>(OwnerOfTemporaries, EventNodeName);
 	Blueprint->EventGraphs.Add(ChildStubGraph);
 	ChildStubGraph->Schema = UEdGraphSchema_K2::StaticClass();
