Unnecessary comment removed

diff --git a/Engine/Source/Editor/KismetCompiler/Private/KismetCompiler.cpp b/Engine/Source/Editor/KismetCompiler/Private/KismetCompiler.cpp
--- a/Engine/Source/Editor/KismetCompiler/Private/KismetCompiler.cpp
+++ b/Engine/Source/Editor/KismetCompiler/Private/KismetCompiler.cpp
@@ -2358,7 +2358,6 @@ void FKismetCompilerContext::CreateFunctionStubForEvent(UK2Node_Event* SrcEventN
 	auto ExistingGraph = static_cast<UObject*>(FindObjectWithOuter(OwnerOfTemporaries, UEdGraph::StaticClass(), EventNodeName));
 	if (ExistingGraph && !ExistingGraph->HasAnyFlags(RF_Transient))
 	{
-		// There is no match, so the function parameters must have changed.  Throw an error, and force them to refresh
 		MessageLog.Error(
 			*FString::Printf(*LOCTEXT("CannotCreateStubForEvent_Error", "Graph named '%s' already exists in '%s'. Another one cannot be generated from @@").ToString()
 			, *EventNodeName.ToString()
