Animation: Add a prototype for automagic transition rules when there is a single player node in the source graph
- Defines the can take transition rule as Player.CurrentTime + CrossfadeDuration >= Player.SequenceLength
- Opt-in on the transition node, when enabled the BP transition rule is ignored and the automatic rule is evaulated directly in C++
[INTEGRATE] Change 2454970

diff --git a/Engine/Source/Editor/AnimGraph/Classes/AnimStateTransitionNode.h b/Engine/Source/Editor/AnimGraph/Classes/AnimStateTransitionNode.h
--- a/Engine/Source/Editor/AnimGraph/Classes/AnimStateTransitionNode.h
+++ b/Engine/Source/Editor/AnimGraph/Classes/AnimStateTransitionNode.h
@@ -17,7 +17,7 @@ class UAnimStateTransitionNode : public UAnimStateNodeBase
 	UPROPERTY()
 	class UEdGraph* BoundGraph;
 
-	// The animation graph for this transition if it uses custom blending (reutrning a pose)
+	// The animation graph for this transition if it uses custom blending (returning a pose)
 	UPROPERTY()
 	class UEdGraph* CustomTransitionGraph;
 
@@ -34,6 +34,10 @@ class UAnimStateTransitionNode : public UAnimStateNodeBase
 	UPROPERTY(EditAnywhere, Config, Category=Transition)
 	TEnumAsByte<ETransitionBlendMode::Type> CrossfadeMode;
 
+	// Try setting the rule automatically based on the player node remaining time and the CrossfadeDuration, ignoring the internal time
+	UPROPERTY(EditAnywhere, Category=Transition)
+	bool bAutomaticRuleBasedOnSequencePlayerInState;
+
 	// What transition logic to use
 	UPROPERTY(EditAnywhere, BlueprintReadWrite, Category=Transition)
 	TEnumAsByte<ETransitionLogicType::Type> LogicType;
diff --git a/Engine/Source/Editor/KismetCompiler/Private/AnimBlueprintCompiler.cpp b/Engine/Source/Editor/KismetCompiler/Private/AnimBlueprintCompiler.cpp
--- a/Engine/Source/Editor/KismetCompiler/Private/AnimBlueprintCompiler.cpp
+++ b/Engine/Source/Editor/KismetCompiler/Private/AnimBlueprintCompiler.cpp
@@ -25,6 +25,7 @@
 #include "AnimGraphNode_StateMachineBase.h"
 #include "AnimGraphNode_StateResult.h"
 #include "AnimGraphNode_SaveCachedPose.h"
+#include "AnimGraphNode_SequencePlayer.h"
 #include "AnimGraphNode_TransitionPoseEvaluator.h"
 #include "AnimGraphNode_TransitionResult.h"
 #include "K2Node_TransitionRuleGetter.h"
@@ -738,6 +739,9 @@ void FAnimBlueprintCompiler::ProcessStateMachine(UAnimGraphNode_StateMachineBase
 	new (NewAnimBlueprintClass->BakedStateMachines) FBakedAnimationStateMachine();
 	FMachineCreator Oven(MessageLog, StateMachineInstance, MachineIndex, NewAnimBlueprintClass);
 
+	// Map of states that contain a single player node (from state root node index to associated sequence player)
+	TMap<int32, UObject*> SimplePlayerStatesMap;
+
 	// Process all the states/transitions
 	for (auto StateNodeIt = StateMachineInstance->EditorStateMachineGraph->Nodes.CreateIterator(); StateNodeIt; ++StateNodeIt)
 	{
@@ -816,6 +820,18 @@ void FAnimBlueprintCompiler::ProcessStateMachine(UAnimGraphNode_StateMachineBase
 				if (UAnimGraphNode_StateResult* AnimGraphResultNode = CastChecked<UAnimationStateGraph>(StateNode->BoundGraph)->GetResultNode())
 				{
 					BakedState.StateRootNodeIndex = ExpandGraphAndProcessNodes(StateNode->BoundGraph, AnimGraphResultNode);
+
+					// See if the state consists of a single sequence player node, and remember the index if so
+					for (UEdGraphPin* TestPin : AnimGraphResultNode->Pins)
+					{
+						if ((TestPin->Direction == EGPD_Input) && (TestPin->LinkedTo.Num() == 1))
+						{
+							if (UAnimGraphNode_SequencePlayer* SequencePlayer = Cast<UAnimGraphNode_SequencePlayer>(TestPin->LinkedTo[0]->GetOwningNode()))
+							{
+								SimplePlayerStatesMap.Add(BakedState.StateRootNodeIndex, MessageLog.FindSourceObject(SequencePlayer));
+							}
+						}
+					}
 				}
 				else
 				{
@@ -897,7 +913,29 @@ void FAnimBlueprintCompiler::ProcessStateMachine(UAnimGraphNode_StateMachineBase
 				MessageLog.Error(*LOCTEXT("TransitionWithNoResult", "@@ has no result node").ToString(), TransitionNode);
 			}
 
+			// Handle automatic time remaining rules
+			Rule.StateSequencePlayerToQueryIndex = INDEX_NONE;
+			if (TransitionNode->bAutomaticRuleBasedOnSequencePlayerInState)
+			{
+				if (UObject* SourceSequencePlayer = SimplePlayerStatesMap.FindRef(BakedState.StateRootNodeIndex))
+				{
+					//@TODO: This should be easier than it is... (the nodes are cloned so the pointers don't natrually match)
+					for (const auto& Pair : AllocatedAnimNodeIndices)
+					{
+						if (MessageLog.FindSourceObject(Pair.Key) == SourceSequencePlayer)
+						{
+							Rule.StateSequencePlayerToQueryIndex = Pair.Value;
+						}
+					}
+				}
+
+				if (Rule.StateSequencePlayerToQueryIndex == INDEX_NONE)
+				{
+					MessageLog.Error(*LOCTEXT("CantAutomaticallyDefineTransitionRule", "@@ must contain a single sequence player in order to use bAutomaticRuleBasedOnSequencePlayerInState on @@").ToString(), StateNode, TransitionNode);
+				}
+			}
 
+			// Handle custom transition graphs
 			Rule.CustomResultNodeIndex = INDEX_NONE;
 			if (UAnimationCustomTransitionGraph* CustomTransitionGraph = Cast<UAnimationCustomTransitionGraph>(TransitionNode->CustomTransitionGraph))
 			{
diff --git a/Engine/Source/Runtime/Engine/Classes/Animation/AnimStateMachineTypes.h b/Engine/Source/Runtime/Engine/Classes/Animation/AnimStateMachineTypes.h
--- a/Engine/Source/Runtime/Engine/Classes/Animation/AnimStateMachineTypes.h
+++ b/Engine/Source/Runtime/Engine/Classes/Animation/AnimStateMachineTypes.h
@@ -171,6 +171,10 @@ struct FBakedStateExitTransition
 	UPROPERTY()
 	int32 TransitionIndex;
 
+	// For transitions that go automatically at the end of the sequence, this is the index for the sole sequence player node in the graph
+	UPROPERTY()
+	int32 StateSequencePlayerToQueryIndex;
+
 	// What the transition rule node needs to return to take this transition (for bidirectional transitions)
 	UPROPERTY()
 	bool bDesiredTransitionReturnValue;
@@ -182,6 +186,7 @@ struct FBakedStateExitTransition
 		: CanTakeDelegateIndex(INDEX_NONE)
 		, CustomResultNodeIndex(INDEX_NONE)
 		, TransitionIndex(INDEX_NONE)
+		, StateSequencePlayerToQueryIndex(INDEX_NONE)
 		, bDesiredTransitionReturnValue(true)
 	{
 	}
diff --git a/Engine/Source/Runtime/Engine/Private/Animation/AnimNode_StateMachine.cpp b/Engine/Source/Runtime/Engine/Private/Animation/AnimNode_StateMachine.cpp
--- a/Engine/Source/Runtime/Engine/Private/Animation/AnimNode_StateMachine.cpp
+++ b/Engine/Source/Runtime/Engine/Private/Animation/AnimNode_StateMachine.cpp
@@ -442,8 +442,31 @@ bool FAnimNode_StateMachine::FindValidTransition(const FAnimationUpdateContext&
 		}
 		else
 		{
-			// Execute it and see if we can take this rule
-			ResultNode->EvaluateGraphExposedInputs.Execute(Context);
+			bool bStillCallEvaluate = true;
+
+			if (TransitionRule.StateSequencePlayerToQueryIndex != INDEX_NONE)
+			{
+				// Simple automatic rule
+				FAnimNode_SequencePlayer* SequencePlayer = GetNodeFromPropertyIndex<FAnimNode_SequencePlayer>(Context.AnimInstance, AnimBlueprintClass, TransitionRule.StateSequencePlayerToQueryIndex);
+				if ((SequencePlayer != nullptr) && (SequencePlayer->Sequence != nullptr))
+				{
+					const float SequenceLength = SequencePlayer->Sequence->GetMaxCurrentTime();
+					const float PlayerTime = SequencePlayer->InternalTimeAccumulator;
+					const float PlayerTimeLeft = SequenceLength - PlayerTime;
+
+					const FAnimationTransitionBetweenStates& TransitionInfo = GetTransitionInfo(TransitionRule.TransitionIndex);
+
+					ResultNode->bCanEnterTransition = (PlayerTimeLeft <= TransitionInfo.CrossfadeDuration);
+
+					bStillCallEvaluate = false;
+				}
+			}
+
+			if (bStillCallEvaluate)
+			{
+				// Execute it and see if we can take this rule
+				ResultNode->EvaluateGraphExposedInputs.Execute(Context);
+			}
 		}
 
 		if (ResultNode->bCanEnterTransition == TransitionRule.bDesiredTransitionReturnValue)
