Blueprints: Add a fast calling mechanism for no-argument thunks into the event graph
- Processing no-argument thunks into the event graph (for notification events and animation EvaluateGraphExposedInputs type calls) now calls directly into the event graph, skipping one function frame and associated invocation time
- Optimized return value processing for void functions in UObject::ProcessInternal (skipping a fptr call)
- Fixed a compile error when USE_UBER_GRAPH_PERSISTENT_FRAME is disabled
[INTEGRATE] Change 2455073

diff --git a/Engine/Source/Editor/KismetCompiler/Private/KismetCompilerVMBackend.cpp b/Engine/Source/Editor/KismetCompiler/Private/KismetCompilerVMBackend.cpp
--- a/Engine/Source/Editor/KismetCompiler/Private/KismetCompilerVMBackend.cpp
+++ b/Engine/Source/Editor/KismetCompiler/Private/KismetCompilerVMBackend.cpp
@@ -638,7 +638,7 @@ public:
 		Writer << EX_EndStructConst;
 	}
 
-	void EmitFunctionCall(FBlueprintCompiledStatement& Statement)
+	void EmitFunctionCall(FKismetCompilerContext& CompilerContext, FKismetFunctionContext& FunctionContext, FBlueprintCompiledStatement& Statement)
 	{
 		UFunction* FunctionToCall = Statement.FunctionToCall;
 		check(FunctionToCall);
@@ -676,6 +676,25 @@ public:
 			// Overwrite RHS(0) text with the state index to kick off
 			check(Statement.RHS[Statement.UbergraphCallIndex]->bIsLiteral);
 			Statement.RHS[Statement.UbergraphCallIndex]->Name = FString::FromInt(OffsetWithinUbergraph);
+
+#if UE_BLUEPRINT_EVENTGRAPH_FASTCALLS
+			// Store optimization data if this is a simple call into the ubergraph
+			if (FunctionContext.bIsSimpleStubGraphWithNoParams)
+			{
+				check(FunctionToCall == CompilerContext.NewClass->UberGraphFunction);
+				check(FunctionToCall->ParmsSize == sizeof(int32));
+
+				if ((FunctionToCall->FirstPropertyToInit == nullptr) && (FunctionToCall->PostConstructLink == nullptr))
+				{
+					FEventGraphFastCallPair& Pair = *new (CompilerContext.NewClass->FastCallPairs) FEventGraphFastCallPair();
+					Pair.FunctionToPatch = FunctionContext.Function;
+					Pair.EventGraphCallOffset = OffsetWithinUbergraph;
+
+					FunctionContext.Function->EventGraphFunction = FunctionToCall;
+					FunctionContext.Function->EventGraphCallOffset = OffsetWithinUbergraph;
+				}
+			}
+#endif
 		}
 
 		// Handle the return value assignment if present
@@ -1233,7 +1252,7 @@ public:
 			Writer << ((Statement.Type == KCST_DebugSite) ? EX_Tracepoint : EX_WireTracepoint);
 			break;
 		case KCST_CallFunction:
-			EmitFunctionCall(Statement);
+			EmitFunctionCall(CompilerContext, FunctionContext, Statement);
 			break;
 		case KCST_CallDelegate:
 			EmitCallDelegate(Statement);
diff --git a/Engine/Source/Runtime/Core/Public/Misc/Build.h b/Engine/Source/Runtime/Core/Public/Misc/Build.h
--- a/Engine/Source/Runtime/Core/Public/Misc/Build.h
+++ b/Engine/Source/Runtime/Core/Public/Misc/Build.h
@@ -241,5 +241,8 @@
 /** Enable UberGraphPersistentFrame feature. It can speed up BP compilation (re-instancing) in editor, but introduce an unnecessary overhead in runtime. */
 #define USE_UBER_GRAPH_PERSISTENT_FRAME 1
 
+/** Enable fast calls for event thunks into an event graph that have no parameters  */
+#define UE_BLUEPRINT_EVENTGRAPH_FASTCALLS 1
+
 #define USE_CIRCULAR_DEPENDENCY_LOAD_DEFERRING 1
 #define USE_DEFERRED_DEPENDENCY_CHECK_VERIFICATION_TESTS (USE_CIRCULAR_DEPENDENCY_LOAD_DEFERRING && 0)
diff --git a/Engine/Source/Runtime/CoreUObject/Private/UObject/Class.cpp b/Engine/Source/Runtime/CoreUObject/Private/UObject/Class.cpp
--- a/Engine/Source/Runtime/CoreUObject/Private/UObject/Class.cpp
+++ b/Engine/Source/Runtime/CoreUObject/Private/UObject/Class.cpp
@@ -3634,7 +3634,11 @@ UFunction::UFunction(const FObjectInitializer& ObjectInitializer, UFunction* InS
 , RepOffset(InRepOffset)
 , RPCId(0)
 , RPCResponseId(0)
-, FirstPropertyToInit(NULL)
+, FirstPropertyToInit(nullptr)
+#if UE_BLUEPRINT_EVENTGRAPH_FASTCALLS
+, EventGraphFunction(nullptr)
+, EventGraphCallOffset(0)
+#endif
 {
 }
 
diff --git a/Engine/Source/Runtime/CoreUObject/Private/UObject/ScriptCore.cpp b/Engine/Source/Runtime/CoreUObject/Private/UObject/ScriptCore.cpp
--- a/Engine/Source/Runtime/CoreUObject/Private/UObject/ScriptCore.cpp
+++ b/Engine/Source/Runtime/CoreUObject/Private/UObject/ScriptCore.cpp
@@ -626,7 +626,15 @@ void UObject::ProcessInternal( FFrame& Stack, RESULT_DECL )
 
 		// Step over the return statement and evaluate the result expression
 		Stack.Code++;
-		Stack.Step(Stack.Object, Result);
+
+		if (*Stack.Code != EX_Nothing)
+		{
+			Stack.Step(Stack.Object, Result);
+		}
+		else
+		{
+			Stack.Code++;
+		}
 
 #if DO_GUARD
 		--Recurse;
@@ -840,6 +848,23 @@ void UObject::ProcessEvent( UFunction* Function, void* Parms )
 
 	CONDITIONAL_SCOPE_CYCLE_COUNTER(STAT_BlueprintTime, ScriptEntryTag == 1);
 
+#if UE_BLUEPRINT_EVENTGRAPH_FASTCALLS
+	// Fast path for ubergraph calls
+	int32 EventGraphParams;
+	if (Function->EventGraphFunction != nullptr)
+	{
+		// Call directly into the event graph, skipping the stub thunk function
+		EventGraphParams = Function->EventGraphCallOffset;
+		Parms = &EventGraphParams;
+		Function = Function->EventGraphFunction;
+
+		// Validate assumptions required for this optimized path (EventGraphFunction should have only been filled out if these held)
+		checkSlow(Function->ParmsSize == sizeof(EventGraphParams));
+		checkSlow(Function->FirstPropertyToInit == nullptr);
+		checkSlow(Function->PostConstructLink == nullptr);
+	}
+#endif
+
 	// Scope required for scoped script stats.
 	{
 		uint8* Frame = NULL;
@@ -1233,7 +1258,7 @@ void UObject::execLetValueOnPersistentFrame(FFrame& Stack, RESULT_DECL)
 
 	Stack.Step(Stack.Object, DestAddress);
 #else
-	checkf(false, TEXT("execLetValueOnPersistentFrame: UberGraphPersistentFrame is not supported by current build!")
+	checkf(false, TEXT("execLetValueOnPersistentFrame: UberGraphPersistentFrame is not supported by current build!"));
 #endif
 }
 IMPLEMENT_VM_FUNCTION(Ex_LetValueOnPersistentFrame, execLetValueOnPersistentFrame);
diff --git a/Engine/Source/Runtime/CoreUObject/Public/UObject/Class.h b/Engine/Source/Runtime/CoreUObject/Public/UObject/Class.h
--- a/Engine/Source/Runtime/CoreUObject/Public/UObject/Class.h
+++ b/Engine/Source/Runtime/CoreUObject/Public/UObject/Class.h
@@ -1176,6 +1176,14 @@ public:
 	/** pointer to first local struct property in this UFunction that contains defaults */
 	UProperty* FirstPropertyToInit;
 
+#if UE_BLUEPRINT_EVENTGRAPH_FASTCALLS
+	// The event graph this function calls in to
+	UFunction* EventGraphFunction;
+
+	// The state offset inside of the event graph
+	int32 EventGraphCallOffset;
+#endif
+
 private:
 	Native Func;
 
diff --git a/Engine/Source/Runtime/CoreUObject/Public/UObject/UObject.h b/Engine/Source/Runtime/CoreUObject/Public/UObject/UObject.h
--- a/Engine/Source/Runtime/CoreUObject/Public/UObject/UObject.h
+++ b/Engine/Source/Runtime/CoreUObject/Public/UObject/UObject.h
@@ -1,11 +1,10 @@
 // Copyright 1998-2015 Epic Games, Inc. All Rights Reserved.
 
 /*=============================================================================
-	UnObjUObject.h: Unreal object base class
+	UObject.h: Unreal object base class
 =============================================================================*/
 
-#ifndef __UNOBJUOBJECT_H__
-#define __UNOBJUOBJECT_H__
+#pragma once
 
 DECLARE_LOG_CATEGORY_EXTERN(LogObj, Log, All);
 
@@ -1154,6 +1153,3 @@ private:
 	*/
 	UObject* CreateEditorOnlyDefaultSubobjectImpl(FName SubobjectName, UClass* ReturnType, bool bTransient = false);
 };
-
-#endif	// __UNOBJUOBJECT_H__
-
diff --git a/Engine/Source/Runtime/Engine/Classes/Engine/BlueprintGeneratedClass.h b/Engine/Source/Runtime/Engine/Classes/Engine/BlueprintGeneratedClass.h
--- a/Engine/Source/Runtime/Engine/Classes/Engine/BlueprintGeneratedClass.h
+++ b/Engine/Source/Runtime/Engine/Classes/Engine/BlueprintGeneratedClass.h
@@ -417,6 +417,17 @@ public:
 #endif
 };
 
+USTRUCT()
+struct ENGINE_API FEventGraphFastCallPair
+{
+	GENERATED_USTRUCT_BODY()
+
+	UPROPERTY()
+	UFunction* FunctionToPatch;
+
+	UPROPERTY()
+	int32 EventGraphCallOffset;
+};
 
 UCLASS()
 class ENGINE_API UBlueprintGeneratedClass : public UClass
@@ -454,6 +465,9 @@ public:
 	UPROPERTY()
 	UFunction* UberGraphFunction;
 
+	UPROPERTY()
+	TArray<FEventGraphFastCallPair> FastCallPairs;
+
 	/** 
 	 * Gets an array of all BPGeneratedClasses (including InClass as 0th element) parents of given generated class 
 	 *
diff --git a/Engine/Source/Runtime/Engine/Private/BlueprintGeneratedClass.cpp b/Engine/Source/Runtime/Engine/Private/BlueprintGeneratedClass.cpp
--- a/Engine/Source/Runtime/Engine/Private/BlueprintGeneratedClass.cpp
+++ b/Engine/Source/Runtime/Engine/Private/BlueprintGeneratedClass.cpp
@@ -80,6 +80,15 @@ void UBlueprintGeneratedClass::PostLoad()
 		}
 	}
 #endif // WITH_EDITORONLY_DATA
+
+#if UE_BLUEPRINT_EVENTGRAPH_FASTCALLS
+	// Patch the fast calls (needed as we can't bump engine version to serialize it directly in UFunction right now)
+	for (const FEventGraphFastCallPair& Pair : FastCallPairs)
+	{
+		Pair.FunctionToPatch->EventGraphFunction = UberGraphFunction;
+		Pair.FunctionToPatch->EventGraphCallOffset = Pair.EventGraphCallOffset;
+	}
+#endif
 }
 
 void UBlueprintGeneratedClass::GetRequiredPreloadDependencies(TArray<UObject*>& DependenciesOut)
@@ -639,6 +648,10 @@ void UBlueprintGeneratedClass::PurgeClass(bool bRecompilingOnLoad)
 
 	UberGraphFramePointerProperty = NULL;
 	UberGraphFunction = NULL;
+
+#if UE_BLUEPRINT_EVENTGRAPH_FASTCALLS
+	FastCallPairs.Empty();
+#endif
 }
 
 void UBlueprintGeneratedClass::Bind()
