FKismetCompilerUtilities::CreatePropertyOnScope checks properties in Super Struct.
https://udn.unrealengine.com/questions/237782/crash-migrating-component-from-bp-to-c.html

diff --git a/Engine/Source/Editor/KismetCompiler/Private/KismetCompilerMisc.cpp b/Engine/Source/Editor/KismetCompiler/Private/KismetCompilerMisc.cpp
--- a/Engine/Source/Editor/KismetCompiler/Private/KismetCompilerMisc.cpp
+++ b/Engine/Source/Editor/KismetCompiler/Private/KismetCompilerMisc.cpp
@@ -654,7 +654,25 @@ UProperty* FKismetCompilerUtilities::CreatePropertyOnScope(UStruct* Scope, const
 	// Check to see if there's already a object on this scope with the same name, and throw an internal compiler error if so
 	// If this happens, it breaks the property link, which causes stack corruption and hard-to-track errors, so better to fail at this point
 	{
-		if (UObject* ExistingObject = FindObject<UObject>(Scope, *PropertyName.ToString(), false))
+		auto CheckIfPropertyNameIsUsed = [](UStruct* Struct, const TCHAR* Name) -> UObject*
+		{
+			if (UObject* ExistingObject = FindObject<UObject>(Struct, Name, false))
+			{
+				return ExistingObject;
+			}
+			
+			if (Struct && !Struct->IsA<UFunction>())
+			{
+				if (auto Field = FindField<UField>(Struct ? Struct->GetSuperStruct() : nullptr, Name))
+				{
+					return Field;
+				}
+			}
+
+			return nullptr;
+		};
+
+		if (UObject* ExistingObject = CheckIfPropertyNameIsUsed(Scope, *PropertyName.ToString()))
 		{
 			MessageLog.Error(*FString::Printf(TEXT("Internal Compiler Error:  Tried to create a property %s in scope %s, but another object of type %s already exists there."), *PropertyName.ToString(), (Scope ? *Scope->GetName() : TEXT("None"))), *ExistingObject->GetFullName(Scope));
 
@@ -664,7 +682,7 @@ UProperty* FKismetCompilerUtilities::CreatePropertyOnScope(UStruct* Scope, const
 			do 
 			{
 				TestNameString = PropertyName.ToString() + FString::Printf(TEXT("_ERROR_DUPLICATE_%d"), Counter++);
-			} while (FindObject<UObject>(Scope, *TestNameString, false) != NULL);
+			} while (CheckIfPropertyNameIsUsed(Scope, *TestNameString) != NULL);
 
 			ValidatedPropertyName = FName(*TestNameString);
 		}
