UberGraphFrameName doesn't cause name collision.

diff --git a/Engine/Source/Editor/KismetCompiler/Private/KismetCompilerMisc.cpp b/Engine/Source/Editor/KismetCompiler/Private/KismetCompilerMisc.cpp
--- a/Engine/Source/Editor/KismetCompiler/Private/KismetCompilerMisc.cpp
+++ b/Engine/Source/Editor/KismetCompiler/Private/KismetCompilerMisc.cpp
@@ -661,7 +661,7 @@ UProperty* FKismetCompilerUtilities::CreatePropertyOnScope(UStruct* Scope, const
 				return ExistingObject;
 			}
 			
-			if (Struct && !Struct->IsA<UFunction>())
+			if (Struct && !Struct->IsA<UFunction>() && (UBlueprintGeneratedClass::GetUberGraphFrameName() != Name))
 			{
 				if (auto Field = FindField<UField>(Struct ? Struct->GetSuperStruct() : nullptr, Name))
 				{
