Fix property name in use error message

diff --git a/Engine/Source/Editor/KismetCompiler/Private/KismetCompilerMisc.cpp b/Engine/Source/Editor/KismetCompiler/Private/KismetCompilerMisc.cpp
--- a/Engine/Source/Editor/KismetCompiler/Private/KismetCompilerMisc.cpp
+++ b/Engine/Source/Editor/KismetCompiler/Private/KismetCompilerMisc.cpp
@@ -674,7 +674,7 @@ UProperty* FKismetCompilerUtilities::CreatePropertyOnScope(UStruct* Scope, const
 
 		if (UObject* ExistingObject = CheckIfPropertyNameIsUsed(Scope, *PropertyName.ToString()))
 		{
-			MessageLog.Error(*FString::Printf(TEXT("Internal Compiler Error:  Tried to create a property %s in scope %s, but another object of type %s already exists there."), *PropertyName.ToString(), (Scope ? *Scope->GetName() : TEXT("None"))), *ExistingObject->GetFullName(Scope));
+			MessageLog.Error(*FString::Printf(TEXT("Internal Compiler Error:  Tried to create a property %s in scope %s, but %s already exists there."), *PropertyName.ToString(), (Scope ? *Scope->GetName() : TEXT("None")), *ExistingObject->GetFullName()));
 
 			// Find a free name, so we can still create the property to make it easier to spot the duplicates, and avoid crashing
 			uint32 Counter = 0;
