[UE-11236] Renaming variables that have since become conflicting (as a result of a like-named native parent property).

diff --git a/Engine/Source/Editor/KismetCompiler/Private/KismetCompiler.cpp b/Engine/Source/Editor/KismetCompiler/Private/KismetCompiler.cpp
--- a/Engine/Source/Editor/KismetCompiler/Private/KismetCompiler.cpp
+++ b/Engine/Source/Editor/KismetCompiler/Private/KismetCompiler.cpp
@@ -429,26 +429,34 @@ bool FKismetCompilerContext::IsNodePure(const UEdGraphNode* Node) const
 
 void FKismetCompilerContext::ValidateVariableNames()
 {
-	TSharedPtr<FKismetNameValidator> ParentBPNameValidator;
-	if( Blueprint->ParentClass != NULL )
+	UClass* ParentClass = Blueprint->ParentClass;
+	if (ParentClass != nullptr)
 	{
-		UBlueprint* ParentBP = Cast<UBlueprint>(Blueprint->ParentClass->ClassGeneratedBy);
-		if( ParentBP != NULL )
+		TSharedPtr<FKismetNameValidator> ParentBPNameValidator;
+		if (UBlueprint* ParentBP = Cast<UBlueprint>(Blueprint->ParentClass->ClassGeneratedBy))
 		{
 			ParentBPNameValidator = MakeShareable(new FKismetNameValidator(ParentBP));
 		}
-	}
 
-	if(ParentBPNameValidator.IsValid())
-	{
-		for (int32 VariableIndex=0; VariableIndex < Blueprint->NewVariables.Num(); ++VariableIndex)
+		for (FBPVariableDescription& VarDesc : Blueprint->NewVariables)
 		{
-			FBPVariableDescription& Variable = Blueprint->NewVariables[VariableIndex];
-			if( ParentBPNameValidator->IsValid(Variable.VarName.ToString()) != EValidatorResult::Ok )
+			FName OldVarName = VarDesc.VarName;
+			FName NewVarName = OldVarName;
+
+			FString VarNameStr = OldVarName.ToString();
+			if (ParentBPNameValidator.IsValid() && (ParentBPNameValidator->IsValid(VarNameStr) != EValidatorResult::Ok))
+			{
+				NewVarName = FBlueprintEditorUtils::FindUniqueKismetName(Blueprint, VarNameStr);
+			}
+			else if (ParentClass->HasAnyFlags(RF_Native) && FindObject<UObject>(ParentClass, *VarNameStr, /*ExactClass =*/false))
+			{
+				NewVarName = FBlueprintEditorUtils::FindUniqueKismetName(Blueprint, VarNameStr);
+			}
+
+			if (OldVarName != NewVarName)
 			{
-				FName NewVariableName = FBlueprintEditorUtils::FindUniqueKismetName(Blueprint, Variable.VarName.ToString());
-				MessageLog.Warning(*FString::Printf(*LOCTEXT("MemberVariableConflictWarning", "Found a member variable with a conflicting name (%s) - changed to %s.").ToString(), *Variable.VarName.ToString(), *NewVariableName.ToString()));
-				FBlueprintEditorUtils::RenameMemberVariable(Blueprint, Variable.VarName, NewVariableName);
+				MessageLog.Warning(*FString::Printf(*LOCTEXT("MemberVariableConflictWarning", "Found a member variable with a conflicting name (%s) - changed to %s.").ToString(), *VarNameStr, *NewVarName.ToString()));
+				FBlueprintEditorUtils::RenameMemberVariable(Blueprint, OldVarName, NewVarName);
 			}
 		}
 	}
