[UE-11272] Raising UndeterminedPinType_Error only on unorphaned nodes. This fixes a bluepring compile error that results when some projects are upgraded to 4.7 (because pin types are reverted to wildcards when possible).

diff --git a/Engine/Source/Editor/KismetCompiler/Private/KismetCompiler.cpp b/Engine/Source/Editor/KismetCompiler/Private/KismetCompiler.cpp
--- a/Engine/Source/Editor/KismetCompiler/Private/KismetCompiler.cpp
+++ b/Engine/Source/Editor/KismetCompiler/Private/KismetCompiler.cpp
@@ -326,12 +326,6 @@ void FKismetCompilerContext::ValidatePin(const UEdGraphPin* Pin) const
 		return;
 	}
 
-	if (Pin->PinType.PinCategory == Schema->PC_Wildcard)
-	{
-		// Wildcard pins should never be seen by the compiler; they should always be forced into a particular type by wiring.
-		MessageLog.Error(*LOCTEXT("UndeterminedPinType_Error", "The type of @@ is undetermined.  Connect something to @@ to imply a specific type.").ToString(), Pin, OwningNodeUnchecked);
-	}
-
 	if (Pin->LinkedTo.Num() > 1)
 	{
 		if (Pin->Direction == EGPD_Output)
@@ -1119,6 +1113,27 @@ void FKismetCompilerContext::ValidateSelfPinsInGraph(const UEdGraph* SourceGraph
 	}
 }
 
+void FKismetCompilerContext::ValidateNoWildcardPinsInGraph(const UEdGraph* SourceGraph)
+{
+	for (int32 NodeIndex = 0; NodeIndex < SourceGraph->Nodes.Num(); ++NodeIndex)
+	{
+		if (const UEdGraphNode* Node = SourceGraph->Nodes[NodeIndex])
+		{
+			for (int32 PinIndex = 0; PinIndex < Node->Pins.Num(); ++PinIndex)
+			{
+				if (const UEdGraphPin* Pin = Node->Pins[PinIndex])
+				{
+					if (Pin->PinType.PinCategory == Schema->PC_Wildcard)
+					{
+						// Wildcard pins should never be seen by the compiler; they should always be forced into a particular type by wiring.
+						MessageLog.Error(*LOCTEXT("UndeterminedPinType_Error", "The type of @@ is undetermined.  Connect something to @@ to imply a specific type.").ToString(), Pin, Pin->GetOwningNodeUnchecked());
+					}
+				}
+			}
+		}
+	}
+}
+
 /**
  * First phase of compiling a function graph
  *   - Prunes the 'graph' to only included the connected portion that contains the function entry point 
@@ -1155,8 +1170,9 @@ void FKismetCompilerContext::PrecompileFunction(FKismetFunctionContext& Context)
 
 		if (bIsFullCompile)
 		{
-			// Check if self pins are connected after PruneIsolatedNodes, to avoid errors from isolated nodes.
+			// Check if self pins are connected and types are resolved after PruneIsolatedNodes, to avoid errors from isolated nodes.
 			ValidateSelfPinsInGraph(Context.SourceGraph);
+			ValidateNoWildcardPinsInGraph(Context.SourceGraph);
 
 			// Transforms
 			TransformNodes(Context);
diff --git a/Engine/Source/Editor/KismetCompiler/Public/KismetCompiler.h b/Engine/Source/Editor/KismetCompiler/Public/KismetCompiler.h
--- a/Engine/Source/Editor/KismetCompiler/Public/KismetCompiler.h
+++ b/Engine/Source/Editor/KismetCompiler/Public/KismetCompiler.h
@@ -370,6 +370,11 @@ protected:
 	 */
 	void ValidateSelfPinsInGraph(const UEdGraph* SourceGraph);
 
+	/**
+	* Checks if pin types are unresolved (e.g. still wildcards).
+	*/
+	void ValidateNoWildcardPinsInGraph(const UEdGraph* SourceGraph);
+
 	/** Ensures that all variables have valid names for compilation/replication */
 	void ValidateVariableNames();
 
