UE-9883 bChangeDefaultValueWithoutReinstancing should be re-enabled

When bChangeDefaultValueWithoutReinstancing optimization is enable CRC is still calculated for values, that can be implicitly changed (or could be not propagated).

diff --git a/Engine/Config/BaseEngine.ini b/Engine/Config/BaseEngine.ini
--- a/Engine/Config/BaseEngine.ini
+++ b/Engine/Config/BaseEngine.ini
@@ -1511,7 +1511,7 @@ PrintStringDuration=2.0
 bUseLocalGraphVariables=false
 bPersistentUberGraphFrame=true
 bReinstanceOnlyWhenNecessary=true
-bChangeDefaultValueWithoutReinstancing=false
+bChangeDefaultValueWithoutReinstancing=true
 bDisplaysLayout=false
 bNativeCodeGenerationTool=false
 bSkeletonInheritSkeletonClasses=false
diff --git a/Engine/Source/Editor/KismetCompiler/Private/KismetCompiler.cpp b/Engine/Source/Editor/KismetCompiler/Private/KismetCompiler.cpp
--- a/Engine/Source/Editor/KismetCompiler/Private/KismetCompiler.cpp
+++ b/Engine/Source/Editor/KismetCompiler/Private/KismetCompiler.cpp
@@ -3748,12 +3748,50 @@ void FKismetCompilerContext::Compile()
 		// TODO What do we do if validation fails?
 	}
 
-	static const FBoolConfigValueHelper ChangeDefaultValueWithoutReinstancing(TEXT("Kismet"), TEXT("bChangeDefaultValueWithoutReinstancing"), GEngineIni);
-	if (bIsFullCompile && !ChangeDefaultValueWithoutReinstancing)
+	if (bIsFullCompile)
 	{
 		BP_SCOPED_COMPILER_EVENT_STAT(EKismetCompilerStats_ChecksumCDO);
+
+		static const FBoolConfigValueHelper ChangeDefaultValueWithoutReinstancing(TEXT("Kismet"), TEXT("bChangeDefaultValueWithoutReinstancing"), GEngineIni);
+		// CRC is usually calculated for all Properties. If the bChangeDefaultValueWithoutReinstancing optimization is enabled, then only specific properties are considered (in fact we should consider only . See UE-9883.
+		// Some native properties (bCanEverTick) may be implicitly changed by KismetCompiler during compilation, so they always need to be compared.
+		// Some properties with a custom Property Editor Widget may not propagate changes among instances. They may be also compared.
+
+		class FSpecializedArchiveCrc32 : public FArchiveObjectCrc32
+		{
+		public:
+			bool bAllProperties;
+
+			FSpecializedArchiveCrc32(bool bInAllProperties)
+				: FArchiveObjectCrc32()
+				, bAllProperties(bInAllProperties)
+			{}
+
+			static bool PropertyCanBeImplicitlyChanged(const UProperty* InProperty)
+			{
+				check(InProperty);
+
+				auto PropertyOwnerClass = InProperty->GetOwnerClass();
+				const bool bOwnerIsNativeClass = PropertyOwnerClass && PropertyOwnerClass->HasAnyClassFlags(CLASS_Native);
+
+				auto PropertyOwnerStruct = InProperty->GetOwnerStruct();
+				const bool bOwnerIsNativeStruct = !PropertyOwnerClass && (!PropertyOwnerStruct || !PropertyOwnerStruct->IsA<UUserDefinedStruct>());
+
+				return InProperty->IsA<UStructProperty>()
+					|| bOwnerIsNativeClass || bOwnerIsNativeStruct;
+			}
+
+			// Begin FArchive Interface
+			virtual bool ShouldSkipProperty(const UProperty* InProperty) const override
+			{
+				return FArchiveObjectCrc32::ShouldSkipProperty(InProperty) 
+					|| (!bAllProperties && !PropertyCanBeImplicitlyChanged(InProperty));
+			}
+			// End FArchive Interface
+		};
+
 		UObject* NewCDO = NewClass->GetDefaultObject(false);
-		FArchiveObjectCrc32 CrcArchive;
+		FSpecializedArchiveCrc32 CrcArchive(!ChangeDefaultValueWithoutReinstancing);
 		Blueprint->CrcPreviousCompiledCDO = NewCDO ? CrcArchive.Crc32(NewCDO) : 0;
 	}
 }
diff --git a/Engine/Source/Editor/UnrealEd/Private/Kismet2/KismetReinstanceUtilities.cpp b/Engine/Source/Editor/UnrealEd/Private/Kismet2/KismetReinstanceUtilities.cpp
--- a/Engine/Source/Editor/UnrealEd/Private/Kismet2/KismetReinstanceUtilities.cpp
+++ b/Engine/Source/Editor/UnrealEd/Private/Kismet2/KismetReinstanceUtilities.cpp
@@ -356,12 +356,9 @@ void FBlueprintCompileReinstancer::ReinstanceInner(bool bForceAlwaysReinstance)
 			const UBlueprintGeneratedClass* BPClassB = Cast<const UBlueprintGeneratedClass>(ClassToReinstance);
 			const UBlueprint* BP = Cast<const UBlueprint>(ClassToReinstance->ClassGeneratedBy);
 
-			static const FBoolConfigValueHelper ChangeDefaultValueWithoutReinstancing(TEXT("Kismet"), TEXT("bChangeDefaultValueWithoutReinstancing"), GEngineIni);
 			const bool bTheSameDefaultValues = (BP != nullptr) && (ClassToReinstanceDefaultValuesCRC != 0) && (BP->CrcPreviousCompiledCDO == ClassToReinstanceDefaultValuesCRC);
-
 			const bool bTheSameLayout = (BPClassA != nullptr) && (BPClassB != nullptr) && FStructUtils::TheSameLayout(BPClassA, BPClassB, true);
-
-			const bool bAllowedToDoFastPath = (ChangeDefaultValueWithoutReinstancing || bTheSameDefaultValues) && bTheSameLayout;
+			const bool bAllowedToDoFastPath = bTheSameDefaultValues && bTheSameLayout;
 			if (bAllowedToDoFastPath)
 			{
 				ReinstanceFast();
