[UE-11325] Blueprint function params marked "DefaultToSelf" now check compatibility against the Blueprint's type (and error when appropriate).

diff --git a/Engine/Source/Editor/BlueprintGraph/Private/CallFunctionHandler.cpp b/Engine/Source/Editor/BlueprintGraph/Private/CallFunctionHandler.cpp
--- a/Engine/Source/Editor/BlueprintGraph/Private/CallFunctionHandler.cpp
+++ b/Engine/Source/Editor/BlueprintGraph/Private/CallFunctionHandler.cpp
@@ -577,7 +577,7 @@ void FKCHandler_CallFunction::RegisterNets(FKismetFunctionContext& Context, UEdG
 			// if this pin could use a default (it doesn't have a connection or default of its own)
 			if (!bIsConnected && (Pin->DefaultObject == NULL))
 			{
-				if (DefaultToSelfParamNames.Contains(Pin->PinName))
+				if (DefaultToSelfParamNames.Contains(Pin->PinName) && FKismetCompilerUtilities::ValidateSelfCompatibility(Pin, Context))
 				{
 					ensure(Pin->PinType.PinSubCategoryObject != NULL);
 					ensure((Pin->PinType.PinCategory == K2Schema->PC_Object) || (Pin->PinType.PinCategory == K2Schema->PC_Interface));
diff --git a/Engine/Source/Editor/KismetCompiler/Private/KismetCompiler.cpp b/Engine/Source/Editor/KismetCompiler/Private/KismetCompiler.cpp
--- a/Engine/Source/Editor/KismetCompiler/Private/KismetCompiler.cpp
+++ b/Engine/Source/Editor/KismetCompiler/Private/KismetCompiler.cpp
@@ -1053,8 +1053,10 @@ void FKismetCompilerContext::PruneIsolatedNodes(const TArray<UEdGraphNode*>& Roo
 /**
  *	Checks if self pins are connected.
  */
-void FKismetCompilerContext::ValidateSelfPinsInGraph(const UEdGraph* SourceGraph)
+void FKismetCompilerContext::ValidateSelfPinsInGraph(FKismetFunctionContext& Context)
 {
+	const UEdGraph* SourceGraph = Context.SourceGraph;
+
 	check(NULL != Schema);
 	for (int32 NodeIndex = 0; NodeIndex < SourceGraph->Nodes.Num(); ++NodeIndex)
 	{
@@ -1066,46 +1068,7 @@ void FKismetCompilerContext::ValidateSelfPinsInGraph(const UEdGraph* SourceGraph
 				{
 					if (Schema->IsSelfPin(*Pin) && (Pin->LinkedTo.Num() == 0) && Pin->DefaultObject == nullptr)
 					{
-						FEdGraphPinType SelfType;
-						SelfType.PinCategory = Schema->PC_Object;
-						SelfType.PinSubCategory = Schema->PSC_Self;
-
-						FString ErrorMsg;
-						if(Blueprint->BlueprintType != BPTYPE_FunctionLibrary && Schema->IsStaticFunctionGraph(SourceGraph))
-						{
-							ErrorMsg = FString::Printf(*LOCTEXT("PinMustHaveConnection_Static_Error", "'@@' must have a connection, because %s is a static function and will not be bound to instances of this blueprint.").ToString(), *SourceGraph->GetName());
-						}
-						else if (!Schema->ArePinTypesCompatible(SelfType, Pin->PinType, NewClass))
-						{
-							FString PinType = Pin->PinType.PinCategory;
-							if ((Pin->PinType.PinCategory == Schema->PC_Object)    || 
-								(Pin->PinType.PinCategory == Schema->PC_Interface) ||
-								(Pin->PinType.PinCategory == Schema->PC_Class))
-							{
-								if (Pin->PinType.PinSubCategoryObject.IsValid())
-								{
-									PinType = Pin->PinType.PinSubCategoryObject->GetName();
-								}
-								else
-								{
-									PinType = TEXT("");
-								}
-							}
-
-							if(PinType.IsEmpty())
-							{
-								ErrorMsg = FString::Printf(*LOCTEXT("PinMustHaveConnection_NoType_Error", "'@@' must have a connection.").ToString());
-							}
-							else
-							{
-								ErrorMsg = FString::Printf(*LOCTEXT("PinMustHaveConnection_WrongClass_Error", "This blueprint (self) is not a %s, therefore '@@' must have a connection.").ToString(), *PinType);	
-							}
-						}
-
-						if(!ErrorMsg.IsEmpty())
-						{
-							MessageLog.Error(*ErrorMsg, Pin);
-						}
+						FKismetCompilerUtilities::ValidateSelfCompatibility(Pin, Context);
 					}
 				}
 			}
@@ -1171,7 +1134,7 @@ void FKismetCompilerContext::PrecompileFunction(FKismetFunctionContext& Context)
 		if (bIsFullCompile)
 		{
 			// Check if self pins are connected and types are resolved after PruneIsolatedNodes, to avoid errors from isolated nodes.
-			ValidateSelfPinsInGraph(Context.SourceGraph);
+			ValidateSelfPinsInGraph(Context);
 			ValidateNoWildcardPinsInGraph(Context.SourceGraph);
 
 			// Transforms
diff --git a/Engine/Source/Editor/KismetCompiler/Private/KismetCompilerMisc.cpp b/Engine/Source/Editor/KismetCompiler/Private/KismetCompilerMisc.cpp
--- a/Engine/Source/Editor/KismetCompiler/Private/KismetCompilerMisc.cpp
+++ b/Engine/Source/Editor/KismetCompiler/Private/KismetCompilerMisc.cpp
@@ -507,6 +507,61 @@ void FKismetCompilerUtilities::ValidateEnumProperties(UObject* DefaultObject, FC
 	}
 }
 
+bool FKismetCompilerUtilities::ValidateSelfCompatibility(const UEdGraphPin* Pin, FKismetFunctionContext& Context)
+{
+	const UBlueprint* Blueprint = Context.Blueprint;
+	const UEdGraph* SourceGraph = Context.SourceGraph;
+	UEdGraphSchema_K2* K2Schema = Context.Schema;
+	const UBlueprintGeneratedClass* BPClass = Context.NewClass;
+
+	FString ErrorMsg;
+	if (Blueprint->BlueprintType != BPTYPE_FunctionLibrary && K2Schema->IsStaticFunctionGraph(SourceGraph))
+	{
+		ErrorMsg = FString::Printf(*LOCTEXT("PinMustHaveConnection_Static_Error", "'@@' must have a connection, because %s is a static function and will not be bound to instances of this blueprint.").ToString(), *SourceGraph->GetName());
+	}
+	else
+	{
+		FEdGraphPinType SelfType;
+		SelfType.PinCategory = K2Schema->PC_Object;
+		SelfType.PinSubCategory = K2Schema->PSC_Self;
+
+		if (!K2Schema->ArePinTypesCompatible(SelfType, Pin->PinType, BPClass))
+		{
+			FString PinType = Pin->PinType.PinCategory;
+			if ((Pin->PinType.PinCategory == K2Schema->PC_Object) ||
+				(Pin->PinType.PinCategory == K2Schema->PC_Interface) ||
+				(Pin->PinType.PinCategory == K2Schema->PC_Class))
+			{
+				if (Pin->PinType.PinSubCategoryObject.IsValid())
+				{
+					PinType = Pin->PinType.PinSubCategoryObject->GetName();
+				}
+				else
+				{
+					PinType = TEXT("");
+				}
+			}
+
+			if (PinType.IsEmpty())
+			{
+				ErrorMsg = FString::Printf(*LOCTEXT("PinMustHaveConnection_NoType_Error", "This blueprint (self) is not compatible with '@@', therefore that pin must have a connection.").ToString());
+			}
+			else
+			{
+				ErrorMsg = FString::Printf(*LOCTEXT("PinMustHaveConnection_WrongClass_Error", "This blueprint (self) is not a %s, therefore '@@' must have a connection.").ToString(), *PinType);
+			}
+		}
+	}
+
+	if (!ErrorMsg.IsEmpty())
+	{
+		Context.MessageLog.Error(*ErrorMsg, Pin);
+		return false;
+	}
+
+	return true;
+}
+
 UEdGraphPin* FKismetCompilerUtilities::GenerateAssignmentNodes(class FKismetCompilerContext& CompilerContext, UEdGraph* SourceGraph, UK2Node_CallFunction* CallBeginSpawnNode, UEdGraphNode* SpawnNode, UEdGraphPin* CallBeginResult, const UClass* ForClass )
 {
 	static FString ObjectParamName = FString(TEXT("Object"));
diff --git a/Engine/Source/Editor/KismetCompiler/Public/KismetCompiler.h b/Engine/Source/Editor/KismetCompiler/Public/KismetCompiler.h
--- a/Engine/Source/Editor/KismetCompiler/Public/KismetCompiler.h
+++ b/Engine/Source/Editor/KismetCompiler/Public/KismetCompiler.h
@@ -368,7 +368,7 @@ protected:
 	/**
 	 * Checks if self pins are connected.
 	 */
-	void ValidateSelfPinsInGraph(const UEdGraph* SourceGraph);
+	void ValidateSelfPinsInGraph(FKismetFunctionContext& Context);
 
 	/**
 	* Checks if pin types are unresolved (e.g. still wildcards).
diff --git a/Engine/Source/Editor/KismetCompiler/Public/KismetCompilerMisc.h b/Engine/Source/Editor/KismetCompiler/Public/KismetCompilerMisc.h
--- a/Engine/Source/Editor/KismetCompiler/Public/KismetCompilerMisc.h
+++ b/Engine/Source/Editor/KismetCompiler/Public/KismetCompilerMisc.h
@@ -76,6 +76,9 @@ public:
 	/* checks if enum variables from given object store proper indexes */
 	static void ValidateEnumProperties(UObject* DefaultObject, FCompilerResultsLog& MessageLog);
 
+	/** checks if the specified pin can default to self */
+	static bool ValidateSelfCompatibility(const UEdGraphPin* Pin, FKismetFunctionContext& Context);
+
 	/** Create 'set var by name' nodes and hook them up - used to set values when components are added or actor are created at run time. Returns the 'last then' pin of the assignment nodes */
 	static UEdGraphPin* GenerateAssignmentNodes( class FKismetCompilerContext& CompilerContext, UEdGraph* SourceGraph, UK2Node_CallFunction* CallBeginSpawnNode, UEdGraphNode* SpawnNode, UEdGraphPin* CallBeginResult, const UClass* ForClass );
 
