Fixes issue during compile on load to rebuild nodes from expanded macro graphs if the macro's BP was not regenerated yet.

diff --git a/Engine/Source/Editor/KismetCompiler/Private/KismetCompiler.cpp b/Engine/Source/Editor/KismetCompiler/Private/KismetCompiler.cpp
--- a/Engine/Source/Editor/KismetCompiler/Private/KismetCompiler.cpp
+++ b/Engine/Source/Editor/KismetCompiler/Private/KismetCompiler.cpp
@@ -2933,10 +2933,17 @@ void FKismetCompilerContext::ExpandTunnelsAndMacros(UEdGraph* SourceGraph)
 					/*out*/ NodeOffsetY);
 			}
 
+			// If regenerating on load and the MacroGraph's Blueprint did not reach PostLoad yet, the nodes expanded from the graph will need to be reconstructed
+			bool bMacroNodesNeedRefresh = Blueprint->bIsRegeneratingOnLoad && !FBlueprintEditorUtils::FindBlueprintForGraph(MacroGraph)->bHasBeenRegenerated;
+
 			// Record intermediate object creation nodes, offset the nodes, and handle tunnels
 			for (TArray<UEdGraphNode*>::TIterator MacroNodeIt(MacroNodes); MacroNodeIt; ++MacroNodeIt)
 			{
 				UEdGraphNode* DuplicatedNode = *MacroNodeIt;
+				if ( bMacroNodesNeedRefresh)
+				{
+					DuplicatedNode->ReconstructNode();
+				}
 
 				if( DuplicatedNode != NULL )
 				{
