UE-5552 Compiling specific blueprint crashes editor

Reinstancer created after an kismet compiler error check BytecodeOnly flag.

diff --git a/Engine/Source/Editor/KismetCompiler/Private/KismetCompilerModule.cpp b/Engine/Source/Editor/KismetCompiler/Private/KismetCompilerModule.cpp
--- a/Engine/Source/Editor/KismetCompiler/Private/KismetCompilerModule.cpp
+++ b/Engine/Source/Editor/KismetCompiler/Private/KismetCompilerModule.cpp
@@ -168,8 +168,8 @@ void FKismet2CompilerModule::CompileBlueprint(class UBlueprint* Blueprint, const
 					ParentReinstancer->ReinstanceObjects();
 				}
 			}
-
-			auto StubReinstancer = FBlueprintCompileReinstancer::Create(Blueprint->GeneratedClass);
+			const bool bBytecodeOnly = EKismetCompileType::BytecodeOnly == CompileOptions.CompileType;
+			auto StubReinstancer = FBlueprintCompileReinstancer::Create(Blueprint->GeneratedClass, bBytecodeOnly);
 
 			// Toss the half-baked class and generate a stubbed out skeleton class that can be used
 			FCompilerResultsLog StubResults;
