Skeleton classes will re-bind and link to their parent BP's skeleton class after reinstancing. Skeleton classes will no longer have a super struct of a REINST class after their parent BP has been structurally modified.

diff --git a/Engine/Source/Editor/KismetCompiler/Private/KismetCompilerModule.cpp b/Engine/Source/Editor/KismetCompiler/Private/KismetCompilerModule.cpp
--- a/Engine/Source/Editor/KismetCompiler/Private/KismetCompilerModule.cpp
+++ b/Engine/Source/Editor/KismetCompiler/Private/KismetCompilerModule.cpp
@@ -126,6 +126,26 @@ void FKismet2CompilerModule::CompileBlueprint(class UBlueprint* Blueprint, const
 		FKismetCompilerOptions SkeletonCompileOptions;
 		SkeletonCompileOptions.CompileType = EKismetCompileType::SkeletonOnly;
 		CompileBlueprintInner(Blueprint, SkeletonCompileOptions, SkeletonResults, ObjLoaded);
+
+		// Relink the child classes
+		TArray<UClass*> ChildrenOfClass;
+		GetDerivedClasses(Blueprint->SkeletonGeneratedClass, ChildrenOfClass);
+		for ( UClass* ChildClass : ChildrenOfClass )
+		{
+			ChildClass->AssembleReferenceTokenStream();
+			ChildClass->SetSuperStruct(Blueprint->SkeletonGeneratedClass);
+
+			auto BPGDuplicatedClass = Cast<UBlueprintGeneratedClass>(ChildClass);
+			auto DuplicatedClassUberGraphFunction = BPGDuplicatedClass ? BPGDuplicatedClass->UberGraphFunction : nullptr;
+			if (DuplicatedClassUberGraphFunction)
+			{
+				DuplicatedClassUberGraphFunction->Bind();
+				DuplicatedClassUberGraphFunction->StaticLink(true);
+			}
+
+			ChildClass->Bind();
+			ChildClass->StaticLink(true);
+		}
 	}
 
 	// If this was a full compile, take appropriate actions depending on the success of failure of the compile
