Fixes potentially unsafe code if a Blueprint cannot be found for a graph. In response to CL# 2471433.

diff --git a/Engine/Source/Editor/KismetCompiler/Private/KismetCompiler.cpp b/Engine/Source/Editor/KismetCompiler/Private/KismetCompiler.cpp
--- a/Engine/Source/Editor/KismetCompiler/Private/KismetCompiler.cpp
+++ b/Engine/Source/Editor/KismetCompiler/Private/KismetCompiler.cpp
@@ -2934,7 +2934,8 @@ void FKismetCompilerContext::ExpandTunnelsAndMacros(UEdGraph* SourceGraph)
 			}
 
 			// If regenerating on load and the MacroGraph's Blueprint did not reach PostLoad yet, the nodes expanded from the graph will need to be reconstructed
-			bool bMacroNodesNeedRefresh = Blueprint->bIsRegeneratingOnLoad && !FBlueprintEditorUtils::FindBlueprintForGraph(MacroGraph)->bHasBeenRegenerated;
+			UBlueprint* MacroBP = FBlueprintEditorUtils::FindBlueprintForGraph(MacroGraph);
+			bool bMacroNodesNeedRefresh = Blueprint->bIsRegeneratingOnLoad && MacroBP && !MacroBP->bHasBeenRegenerated;
 
 			// Record intermediate object creation nodes, offset the nodes, and handle tunnels
 			for (TArray<UEdGraphNode*>::TIterator MacroNodeIt(MacroNodes); MacroNodeIt; ++MacroNodeIt)
