Follow up to [UE-11236], we should do this check only against UProperties, otherwise we generate a variable that cannot be deleted...

diff --git a/Engine/Source/Editor/KismetCompiler/Private/KismetCompilerMisc.cpp b/Engine/Source/Editor/KismetCompiler/Private/KismetCompilerMisc.cpp
--- a/Engine/Source/Editor/KismetCompiler/Private/KismetCompilerMisc.cpp
+++ b/Engine/Source/Editor/KismetCompiler/Private/KismetCompilerMisc.cpp
@@ -718,7 +718,7 @@ UProperty* FKismetCompilerUtilities::CreatePropertyOnScope(UStruct* Scope, const
 			
 			if (Struct && !Struct->IsA<UFunction>() && (UBlueprintGeneratedClass::GetUberGraphFrameName() != Name))
 			{
-				if (auto Field = FindField<UField>(Struct ? Struct->GetSuperStruct() : nullptr, Name))
+				if (auto Field = FindField<UProperty>(Struct ? Struct->GetSuperStruct() : nullptr, Name))
 				{
 					return Field;
 				}
