Explicitly name CDO objects when moving them into the transient package for eventual cleanup during BP compilation.

change summary:
- modified CleanAndSanitizeClass() to pass an explicit, unique object name into Rename() for the old CDO object

diff --git a/Engine/Source/Editor/KismetCompiler/Private/KismetCompiler.cpp b/Engine/Source/Editor/KismetCompiler/Private/KismetCompiler.cpp
--- a/Engine/Source/Editor/KismetCompiler/Private/KismetCompiler.cpp
+++ b/Engine/Source/Editor/KismetCompiler/Private/KismetCompiler.cpp
@@ -174,7 +174,7 @@ void FKismetCompilerContext::CleanAndSanitizeClass(UBlueprintGeneratedClass* Cla
 
 	if( OldCDO )
 	{
-		OldCDO->Rename(NULL, GetTransientPackage(), RenFlags);
+		OldCDO->Rename(*MakeUniqueObjectName(GetTransientPackage(), OldCDO->GetClass(), FName(*FString::Printf(TEXT("TRASH_%s"), *OldCDO->GetClass()->GetDefaultObjectName().ToString()))).ToString(), GetTransientPackage(), RenFlags);
 		ULinkerLoad::InvalidateExport(OldCDO);
 	}
 
