Set outgoing transient CDO rename during Blueprint precompile phase to be more easily identifiable (for debugging).

diff --git a/Engine/Source/Editor/KismetCompiler/Private/KismetCompiler.cpp b/Engine/Source/Editor/KismetCompiler/Private/KismetCompiler.cpp
--- a/Engine/Source/Editor/KismetCompiler/Private/KismetCompiler.cpp
+++ b/Engine/Source/Editor/KismetCompiler/Private/KismetCompiler.cpp
@@ -174,7 +174,9 @@ void FKismetCompilerContext::CleanAndSanitizeClass(UBlueprintGeneratedClass* Cla
 
 	if( OldCDO )
 	{
-		OldCDO->Rename(*MakeUniqueObjectName(GetTransientPackage(), OldCDO->GetClass(), FName(*FString::Printf(TEXT("TRASH_%s"), *OldCDO->GetClass()->GetDefaultObjectName().ToString()))).ToString(), GetTransientPackage(), RenFlags);
+		FString TransientCDOString = FString::Printf(TEXT("TRASH_%s"), *OldCDO->GetName());
+		FName TransientCDOName = MakeUniqueObjectName(GetTransientPackage(), TransientClass, FName(*TransientCDOString));
+		OldCDO->Rename(*TransientCDOName.ToString(), GetTransientPackage(), RenFlags);
 		ULinkerLoad::InvalidateExport(OldCDO);
 	}
 
