Added some asserts to try and trap where a Blueprint CDO is being mutated, and later causing issues with network replication

diff --git a/Engine/Source/Editor/KismetCompiler/Private/KismetCompiler.cpp b/Engine/Source/Editor/KismetCompiler/Private/KismetCompiler.cpp
--- a/Engine/Source/Editor/KismetCompiler/Private/KismetCompiler.cpp
+++ b/Engine/Source/Editor/KismetCompiler/Private/KismetCompiler.cpp
@@ -1819,6 +1819,12 @@ void FKismetCompilerContext::FinishCompilingClass(UClass* Class)
 
 	// Create the default object for this class
 	FKismetCompilerUtilities::CompileDefaultProperties(Class);
+
+	AActor* ActorCDO = Cast<AActor>(Class->GetDefaultObject());
+	if (ActorCDO)
+	{
+		ensureMsg(!ActorCDO->bExchangedRoles, TEXT("Your CDO has had ExchangeNetRoles called on it (likely via RerunConstructionScripts) which should never have happened. This will cause issues replicating this actor over the network due to mutated transient data!"));
+	}
 }
 
 void FKismetCompilerContext::BuildDynamicBindingObjects(UBlueprintGeneratedClass* Class)
diff --git a/Engine/Source/Runtime/Engine/Private/Actor.cpp b/Engine/Source/Runtime/Engine/Private/Actor.cpp
--- a/Engine/Source/Runtime/Engine/Private/Actor.cpp
+++ b/Engine/Source/Runtime/Engine/Private/Actor.cpp
@@ -2452,6 +2452,8 @@ void AActor::PostNetInit()
 
 void AActor::ExchangeNetRoles(bool bRemoteOwned)
 {
+	checkf(!HasAnyFlags(RF_ClassDefaultObject), TEXT("ExchangeNetRoles should never be called on a CDO as it causes issues when replicating actors over the network due to mutated transient data!"));
+
 	if (!bExchangedRoles)
 	{
 		if (bRemoteOwned)
diff --git a/Engine/Source/Runtime/Engine/Private/ActorConstruction.cpp b/Engine/Source/Runtime/Engine/Private/ActorConstruction.cpp
--- a/Engine/Source/Runtime/Engine/Private/ActorConstruction.cpp
+++ b/Engine/Source/Runtime/Engine/Private/ActorConstruction.cpp
@@ -146,6 +146,8 @@ void AActor::DestroyConstructedComponents()
 
 void AActor::RerunConstructionScripts()
 {
+	checkf(!HasAnyFlags(RF_ClassDefaultObject), TEXT("RerunConstructionScripts should never be called on a CDO as it can mutate the transient data on the CDO which then propagates to instances!"));
+
 	FEditorScriptExecutionGuard ScriptGuard;
 	// don't allow (re)running construction scripts on dying actors
 	bool bAllowReconstruction = !IsPendingKill() && !HasAnyFlags(RF_BeginDestroyed|RF_FinishDestroyed);
