UE-12004 Hit an ensure during BP compile

Limited and necessary support for static arrays in BP Compiler.

diff --git a/Engine/Source/Editor/KismetCompiler/Private/KismetCompilerVMBackend.cpp b/Engine/Source/Editor/KismetCompiler/Private/KismetCompilerVMBackend.cpp
--- a/Engine/Source/Editor/KismetCompiler/Private/KismetCompilerVMBackend.cpp
+++ b/Engine/Source/Editor/KismetCompiler/Private/KismetCompilerVMBackend.cpp
@@ -359,7 +359,7 @@ public:
 		}
 	}
 
-	virtual void EmitTermExpr(FBPTerminal* Term, UProperty* CoerceProperty = NULL)
+	virtual void EmitTermExpr(FBPTerminal* Term, UProperty* CoerceProperty = NULL, bool bAllowStaticArray = false)
 	{
 		//@TODO: Must have a coercion type if it's a literal, because the symbol table isn't plumbed in here and the literals don't carry type information either, yay!
 		check((!Term->bIsLiteral) || (CoerceProperty != NULL));
@@ -483,7 +483,10 @@ public:
 					int32 StructSize = Struct->GetStructureSize() * StructProperty->ArrayDim;
 					uint8* StructData = (uint8*)FMemory_Alloca(StructSize);
 					StructProperty->InitializeValue(StructData);
-					ensure(1 == StructProperty->ArrayDim);
+					if (!ensure(bAllowStaticArray || 1 == StructProperty->ArrayDim))
+					{
+						UE_LOG(LogK2Compiler, Error, TEXT("Unsupported static array. Property: %s, Struct: %s"), *StructProperty->GetName(), *Struct->GetName());
+					}
 					if(!FStructureEditorUtils::Fill_MakeStructureDefaultValue(Cast<UUserDefinedStruct>(Struct), StructData))
 					{
 						UE_LOG(LogK2Compiler, Warning, TEXT("MakeStructureDefaultValue parsing error. Property: %s, Struct: %s"), *StructProperty->GetName(), *Struct->GetName());
@@ -504,20 +507,23 @@ public:
 							continue;
 						}
 
-						// Create a new term for each property, and serialize it out
-						FBPTerminal NewTerm;
-						NewTerm.bIsLiteral = true;
-						Prop->ExportText_InContainer(0, NewTerm.Name, StructData, StructData, NULL, PPF_None);
-						if (Prop->IsA(UTextProperty::StaticClass()))
+						for (int32 ArrayIter = 0; ArrayIter < Prop->ArrayDim; ++ArrayIter)
 						{
-							NewTerm.TextLiteral = FText::FromString(NewTerm.Name);
-						}
-						else if(Prop->IsA(UObjectProperty::StaticClass()))
-						{
-							NewTerm.ObjectLiteral = Cast<UObjectProperty>(Prop)->GetObjectPropertyValue(Prop->ContainerPtrToValuePtr<void>(StructData));
-						}
+							// Create a new term for each property, and serialize it out
+							FBPTerminal NewTerm;
+							NewTerm.bIsLiteral = true;
+							Prop->ExportText_InContainer(ArrayIter, NewTerm.Name, StructData, StructData, NULL, PPF_None);
+							if (Prop->IsA(UTextProperty::StaticClass()))
+							{
+								NewTerm.TextLiteral = FText::FromString(NewTerm.Name);
+							}
+							else if (Prop->IsA(UObjectProperty::StaticClass()))
+							{
+								NewTerm.ObjectLiteral = Cast<UObjectProperty>(Prop)->GetObjectPropertyValue(Prop->ContainerPtrToValuePtr<void>(StructData));
+							}
 
-						EmitTermExpr(&NewTerm, Prop);
+							EmitTermExpr(&NewTerm, Prop, true);
+						}
 					}
 
 					Writer << EX_EndStructConst;
diff --git a/Engine/Source/Runtime/CoreUObject/Private/UObject/ScriptCore.cpp b/Engine/Source/Runtime/CoreUObject/Private/UObject/ScriptCore.cpp
--- a/Engine/Source/Runtime/CoreUObject/Private/UObject/ScriptCore.cpp
+++ b/Engine/Source/Runtime/CoreUObject/Private/UObject/ScriptCore.cpp
@@ -1846,7 +1846,10 @@ void UObject::execStructConst( FFrame& Stack, RESULT_DECL )
 			continue;
 		}
 
-		Stack.Step(Stack.Object, StructProp->ContainerPtrToValuePtr<uint8>(Result));
+		for (int32 ArrayIter = 0; ArrayIter < StructProp->ArrayDim; ++ArrayIter)
+		{
+			Stack.Step(Stack.Object, StructProp->ContainerPtrToValuePtr<uint8>(Result, ArrayIter));
+		}
 	}
 
 	P_FINISH;	// EX_EndStructConst
