#ue4 UEBP-40 - Native code generation from BP, NativeCodeGenerationTool - Various fixes and improvements - WIP

diff --git a/Engine/Source/Editor/Kismet/Private/NativeCodeGenerationTool.cpp b/Engine/Source/Editor/Kismet/Private/NativeCodeGenerationTool.cpp
--- a/Engine/Source/Editor/Kismet/Private/NativeCodeGenerationTool.cpp
+++ b/Engine/Source/Editor/Kismet/Private/NativeCodeGenerationTool.cpp
@@ -6,20 +6,28 @@
 #include "Engine/BlueprintGeneratedClass.h"
 #include "Engine/UserDefinedEnum.h"
 #include "Engine/UserDefinedStruct.h"
+#include "SourceCodeNavigation.h"
+#include "DesktopPlatformModule.h"
 
 #define LOCTEXT_NAMESPACE "NativeCodeGenerationTool"
 
+//
+//	THE CODE SHOULD BE MOVED TO GAMEPROJECTGENERATION
+//
+
+
 struct FGeneratedCodeData
 {
-	FGeneratedCodeData(UBlueprint& Blueprint) : HeaderSource(new FString()), CppSource(new FString())
+	FGeneratedCodeData(UBlueprint& InBlueprint) 
+		: HeaderSource(new FString())
+		, CppSource(new FString())
+		, Blueprint(&InBlueprint)
 	{
 		FName GeneratedClassName, SkeletonClassName;
-		Blueprint.GetBlueprintClassNames(GeneratedClassName, SkeletonClassName);
+		InBlueprint.GetBlueprintClassNames(GeneratedClassName, SkeletonClassName);
 		ClassName = GeneratedClassName.ToString();
 
-		GatherUserDefinedDependencies(Blueprint);
-		
-		FKismetEditorUtilities::GenerateCppCode(&Blueprint, HeaderSource, CppSource);
+		GatherUserDefinedDependencies(InBlueprint);
 	}
 
 	FString TypeDependencies;
@@ -27,6 +35,7 @@ struct FGeneratedCodeData
 	TSharedPtr<FString> CppSource;
 	FString ErrorString;
 	FString ClassName;
+	TWeakObjectPtr<UBlueprint> Blueprint;
 
 	void GatherUserDefinedDependencies(UBlueprint& Blueprint)
 	{
@@ -74,7 +83,7 @@ struct FGeneratedCodeData
 	static FString DefaultHeaderDir()
 	{
 		auto DefaultSourceDir = FPaths::ConvertRelativePathToFull(FPaths::GameSourceDir());
-		return FPaths::Combine(*DefaultSourceDir, FApp::GetGameName(), TEXT("Classes"));
+		return FPaths::Combine(*DefaultSourceDir, FApp::GetGameName(), TEXT("Public"));
 	}
 
 	static FString DefaultSourceDir()
@@ -95,18 +104,89 @@ struct FGeneratedCodeData
 
 	bool Save(const FString& HeaderDirPath, const FString& CppDirPath)
 	{
-		const bool bHeaderSaved = FFileHelper::SaveStringToFile(*HeaderSource, *FPaths::Combine(*HeaderDirPath, *HeaderFileName()));
+		FScopedSlowTask SlowTask(7, LOCTEXT("GeneratingCppFiles", "Generating C++ files.."));
+		SlowTask.MakeDialog();
+		SlowTask.EnterProgressFrame();
+
+		if (!Blueprint.IsValid())
+		{
+			ErrorString += LOCTEXT("InvalidBlueprint", "Invalid Blueprint\n").ToString();
+			return false;
+		}
+
+		FKismetEditorUtilities::GenerateCppCode(Blueprint.Get(), HeaderSource, CppSource, ClassName);
+
+		SlowTask.EnterProgressFrame();
+
+		bool bProjectHadCodeFiles = false;
+		{
+			TArray<FString> OutProjectCodeFilenames;
+			IFileManager::Get().FindFilesRecursive(OutProjectCodeFilenames, *FPaths::GameSourceDir(), TEXT("*.h"), true, false, false);
+			IFileManager::Get().FindFilesRecursive(OutProjectCodeFilenames, *FPaths::GameSourceDir(), TEXT("*.cpp"), true, false, false);
+			bProjectHadCodeFiles = OutProjectCodeFilenames.Num() > 0;
+		}
+
+		TArray<FString> CreatedFiles;
+
+		const FString NeHeaderFilename = FPaths::Combine(*HeaderDirPath, *HeaderFileName());
+		const bool bHeaderSaved = FFileHelper::SaveStringToFile(*HeaderSource, *NeHeaderFilename);
 		if (!bHeaderSaved)
 		{
 			ErrorString += LOCTEXT("HeaderNotSaved", "Header file wasn't saved. Check log for details.\n").ToString();
 		}
-		const bool bCppSaved = FFileHelper::SaveStringToFile(*CppSource, *FPaths::Combine(*CppDirPath, *SourceFileName()));
+		else
+		{
+			CreatedFiles.Add(NeHeaderFilename);
+		}
+
+		SlowTask.EnterProgressFrame();
+
+		const FString NewCppFilename = FPaths::Combine(*CppDirPath, *SourceFileName());
+		const bool bCppSaved = FFileHelper::SaveStringToFile(*CppSource, *NewCppFilename);
 		if (!bCppSaved)
 		{
 			ErrorString += LOCTEXT("CppNotSaved", "Cpp file wasn't saved. Check log for details.\n").ToString();
 		}
+		else
+		{
+			CreatedFiles.Add(NewCppFilename);
+		}
+
+		TArray<FString> CreatedFilesForExternalAppRead;
+		CreatedFilesForExternalAppRead.Reserve(CreatedFiles.Num());
+		for (const FString& CreatedFile : CreatedFiles)
+		{
+			CreatedFilesForExternalAppRead.Add(IFileManager::Get().ConvertToAbsolutePathForExternalAppForRead(*CreatedFile));
+		}
+
+		SlowTask.EnterProgressFrame();
+
+		bool bGenerateProjectFiles = true;
+		// First see if we can avoid a full generation by adding the new files to an already open project
+		if (bProjectHadCodeFiles && FSourceCodeNavigation::AddSourceFiles(CreatedFilesForExternalAppRead))
+		{
+			// We successfully added the new files to the solution, but we still need to run UBT with -gather to update any UBT makefiles
+			if (FDesktopPlatformModule::Get()->InvalidateMakefiles(FPaths::RootDir(), FPaths::GetProjectFilePath(), GWarn))
+			{
+				// We managed the gather, so we can skip running the full generate
+				bGenerateProjectFiles = false;
+			}
+		}
 
-		return bHeaderSaved && bCppSaved;
+		SlowTask.EnterProgressFrame();
+
+		bool bProjectFileUpdated = true;
+		if (bGenerateProjectFiles)
+		{
+			// Generate project files if we happen to be using a project file.
+			if (!FDesktopPlatformModule::Get()->GenerateProjectFiles(FPaths::RootDir(), FPaths::GetProjectFilePath(), GWarn))
+			{
+				ErrorString += LOCTEXT("FailedToGenerateProjectFiles", "Failed to generate project files.").ToString();
+				bProjectFileUpdated = false;
+			}
+		}
+
+		return bHeaderSaved && bCppSaved && bProjectFileUpdated;
 	}
 };
 
@@ -235,6 +315,27 @@ private:
 		return IsEditable() ? LOCTEXT("Generate", "Generate") : LOCTEXT("Close", "Close");
 	}
 
+	FText GetClassName() const
+	{
+		return GeneratedCodeData.IsValid() ? FText::FromString(GeneratedCodeData->ClassName) : FText::GetEmpty();
+	}
+
+	void OnClassNameCommited(const FText& NameText, ETextCommit::Type TextCommit)
+	{
+		if (GeneratedCodeData.IsValid())
+		{
+			GeneratedCodeData->ClassName = NameText.ToString();
+			if (HeaderDirectoryBrowser.IsValid())
+			{
+				HeaderDirectoryBrowser->File = GeneratedCodeData->HeaderFileName();
+			}
+			if (SourceDirectoryBrowser.IsValid())
+			{
+				SourceDirectoryBrowser->File = GeneratedCodeData->SourceFileName();
+			}
+		}
+	}
+
 public:
 
 	void Construct(const FArguments& InArgs)
@@ -253,6 +354,22 @@ public:
 				+ SVerticalBox::Slot()
 				.Padding(4.0f)
 				.AutoHeight()
+				[
+					SNew(STextBlock)
+					.Text(LOCTEXT("ClassName", "Class Name"))
+				]
+				+ SVerticalBox::Slot()
+				.Padding(4.0f)
+				.AutoHeight()
+				[
+					SNew(SEditableTextBox)
+					.Text(this, &SNativeCodeGenerationDialog::GetClassName)
+					.OnTextCommitted(this, &SNativeCodeGenerationDialog::OnClassNameCommited)
+					.IsEnabled(this, &SNativeCodeGenerationDialog::IsEditable)
+				]
+				+ SVerticalBox::Slot()
+				.Padding(4.0f)
+				.AutoHeight()
 				[
 					SNew(STextBlock)
 					.Text(LOCTEXT("HeaderPath", "Header Path"))
diff --git a/Engine/Source/Editor/KismetCompiler/Private/KismetCompiler.cpp b/Engine/Source/Editor/KismetCompiler/Private/KismetCompiler.cpp
--- a/Engine/Source/Editor/KismetCompiler/Private/KismetCompiler.cpp
+++ b/Engine/Source/Editor/KismetCompiler/Private/KismetCompiler.cpp
@@ -3598,31 +3598,6 @@ void FKismetCompilerContext::Compile()
 			GConfig->GetBool(TEXT("Kismet"), TEXT("CompileDisplaysBinaryBackend"), /*out*/ bDisplayBytecode, GEngineIni);
 		}
 
-		// Generate code thru the backend(s)
-		if ((bDisplayCpp && bIsFullCompile) || CompileOptions.DoesRequireCppCodeGeneration())
-		{
-			TUniquePtr<IKismetCppBackend> Backend_CPP(IKismetCppBackend::Create(Schema, *this));
-
-			// The C++ backend is currently only for debugging, so it's only run if the output will be visible
-			Backend_CPP->GenerateCodeFromClass(NewClass, FunctionList, !bIsFullCompile);
-		
-			if (CompileOptions.OutHeaderSourceCode.IsValid())
-			{
-				*CompileOptions.OutHeaderSourceCode = Backend_CPP->GetHeader();
-			}
-
-			if (CompileOptions.OutCppSourceCode.IsValid())
-			{
-				*CompileOptions.OutCppSourceCode = Backend_CPP->GetBody();
-			}
-
-			if (bDisplayCpp)
-			{
-				UE_LOG(LogK2Compiler, Log, TEXT("[header]\n\n\n%s"), *Backend_CPP->GetHeader());
-				UE_LOG(LogK2Compiler, Log, TEXT("[body]\n\n\n%s"), *Backend_CPP->GetBody());
-			}
-		}
-
 		// Always run the VM backend, it's needed for more than just debug printing
 		{
 			BP_SCOPED_COMPILER_EVENT_STAT(EKismetCompilerStats_CodeGenerationTime);
@@ -3648,6 +3623,31 @@ void FKismetCompilerContext::Compile()
 			}
 		}
 
+		// Generate code thru the backend(s)
+		if ((bDisplayCpp && bIsFullCompile) || CompileOptions.DoesRequireCppCodeGeneration())
+		{
+			TUniquePtr<IKismetCppBackend> Backend_CPP(IKismetCppBackend::Create(Schema, *this));
+
+			// The C++ backend is currently only for debugging, so it's only run if the output will be visible
+			Backend_CPP->GenerateCodeFromClass(NewClass, CompileOptions.NewCppClassName, FunctionList, !bIsFullCompile);
+
+			if (CompileOptions.OutHeaderSourceCode.IsValid())
+			{
+				*CompileOptions.OutHeaderSourceCode = Backend_CPP->GetHeader();
+			}
+
+			if (CompileOptions.OutCppSourceCode.IsValid())
+			{
+				*CompileOptions.OutCppSourceCode = Backend_CPP->GetBody();
+			}
+
+			if (bDisplayCpp)
+			{
+				UE_LOG(LogK2Compiler, Log, TEXT("[header]\n\n\n%s"), *Backend_CPP->GetHeader());
+				UE_LOG(LogK2Compiler, Log, TEXT("[body]\n\n\n%s"), *Backend_CPP->GetBody());
+			}
+		}
+
 		static const FBoolConfigValueHelper DisplayLayout(TEXT("Kismet"), TEXT("bDisplaysLayout"), GEngineIni);
 		if (!Blueprint->bIsRegeneratingOnLoad && bIsFullCompile && DisplayLayout && NewClass)
 		{
diff --git a/Engine/Source/Editor/KismetCompiler/Private/KismetCompilerBackend.h b/Engine/Source/Editor/KismetCompiler/Private/KismetCompilerBackend.h
--- a/Engine/Source/Editor/KismetCompiler/Private/KismetCompilerBackend.h
+++ b/Engine/Source/Editor/KismetCompiler/Private/KismetCompilerBackend.h
@@ -25,7 +25,7 @@ public:
 
 	virtual const FString& GetBody() const = 0;
 	virtual const FString& GetHeader() const = 0;
-	virtual void GenerateCodeFromClass(UClass* SourceClass, TIndirectArray<FKismetFunctionContext>& Functions, bool bGenerateStubsOnly = false) = 0;
+	virtual void GenerateCodeFromClass(UClass* SourceClass, FString NewClassName, TIndirectArray<FKismetFunctionContext>& Functions, bool bGenerateStubsOnly = false) = 0;
 
 	virtual ~IKismetCppBackend() {}
 };
diff --git a/Engine/Source/Editor/KismetCompiler/Private/KismetCompilerCppBackend.cpp b/Engine/Source/Editor/KismetCompiler/Private/KismetCompilerCppBackend.cpp
--- a/Engine/Source/Editor/KismetCompiler/Private/KismetCompilerCppBackend.cpp
+++ b/Engine/Source/Editor/KismetCompiler/Private/KismetCompilerCppBackend.cpp
@@ -10,6 +10,8 @@
 
 #include "DefaultValueHelper.h"
 
+#include "SourceCodeNavigation.h"
+
 //////////////////////////////////////////////////////////////////////////
 // FKismetCppBackend
 
@@ -93,7 +95,7 @@ public:
 
 	void EmitClassProperties(FStringOutputDevice& Target, UClass* SourceClass);
 
-	void GenerateCodeFromClass(UClass* SourceClass, TIndirectArray<FKismetFunctionContext>& Functions, bool bGenerateStubsOnly) override;
+	void GenerateCodeFromClass(UClass* SourceClass, FString NewClassName, TIndirectArray<FKismetFunctionContext>& Functions, bool bGenerateStubsOnly) override;
 
 	void EmitCallStatment(FKismetFunctionContext& FunctionContext, FBlueprintCompiledStatement& Statement);
 	void EmitCallDelegateStatment(FKismetFunctionContext& FunctionContext, FBlueprintCompiledStatement& Statement);
@@ -433,6 +435,15 @@ struct FEmitHelper
 		return FString();
 	}
 
+	static bool MetaDataCanBeNative(const FName MetaDataName)
+	{
+		if (MetaDataName == TEXT("ModuleRelativePath"))
+		{
+			return false;
+		}
+		return true;
+	}
+
 	static FString HandleMetaData(const UField* Field, bool AddCategory = true)
 	{
 		FString MetaDataStr;
@@ -448,6 +459,10 @@ struct FEmitHelper
 		{
 			for (auto& Pair : *ValuesMap)
 			{
+				if (!MetaDataCanBeNative(Pair.Key))
+				{
+					continue;
+				}
 				if (!Pair.Value.IsEmpty())
 				{
 					MetaDataStrings.Emplace(FString::Printf(TEXT("%s=\"%s\""), *Pair.Key.ToString(), *Pair.Value));
@@ -558,7 +573,7 @@ struct FEmitHelper
 		TArray<FString> Tags;
 
 		//Pointless: BlueprintNativeEvent, BlueprintImplementableEvent
-		//Pointless: CustomThunk
+		//Pointless: CustomThunk ?
 
 		//TODO: SealedEvent
 		//TODO: Unreliable
@@ -580,6 +595,16 @@ struct FEmitHelper
 
 #undef HANDLE_CPF_TAG
 
+	static bool IsBlueprintNativeEvent(uint64 FunctionFlags)
+	{
+		return HasAllFlags(FunctionFlags, FUNC_Event | FUNC_BlueprintEvent | FUNC_Native);
+	}
+
+	static bool IsBlueprintImplementableEvent(uint64 FunctionFlags)
+	{
+		return HasAllFlags(FunctionFlags, FUNC_Event | FUNC_BlueprintEvent) && !HasAllFlags(FunctionFlags, FUNC_Native);
+	}
+
 	static FString EmitUFuntion(UFunction* Function)
 	{
 		TArray<FString> Tags = FEmitHelper::FunctionFlagsToTags(Function->FunctionFlags);
@@ -660,17 +685,10 @@ struct FEmitHelper
 		return Results;
 	}
 
-	static FString MakeCppClassName(UClass* SourceClass)
-	{
-		check(SourceClass);
-		return FString(SourceClass->GetPrefixCPP()) + SourceClass->GetName();
-	}
-
-	static FString EmitLifetimeReplicatedPropsImpl(UClass* SourceClass, const TCHAR* InCurrentIndent)
+	static FString EmitLifetimeReplicatedPropsImpl(UClass* SourceClass, const FString& CppClassName, const TCHAR* InCurrentIndent)
 	{
 		FString Result;
 		bool bFunctionInitilzed = false;
-		const FString CppClassName = MakeCppClassName(SourceClass);
 		for (TFieldIterator<UProperty> It(SourceClass, EFieldIteratorFlags::ExcludeSuper); It; ++It)
 		{
 			if ((It->PropertyFlags & CPF_Net) != 0)
@@ -691,6 +709,56 @@ struct FEmitHelper
 		}
 		return Result;
 	}
+
+	static FString GatherHeadersToInclude(UClass* SourceClass)
+	{
+		TSet<FString> HeaderFiles;
+		{
+			TArray<UObject*> ReferencedObjects;
+			{
+				FReferenceFinder ReferenceFinder(ReferencedObjects, NULL, false, false, false, true);
+				TArray<UObject*> ObjectsToCheck;
+				GetObjectsWithOuter(SourceClass, ObjectsToCheck, true);
+				//CDO?
+				for (auto Obj : ObjectsToCheck)
+				{
+					ReferenceFinder.FindReferences(Obj);
+				}
+			}
+
+			auto EngineSourceDir = FPaths::EngineSourceDir();
+			auto GameSourceDir = FPaths::GameSourceDir();
+			auto CurrentPackage = SourceClass->GetTypedOuter<UPackage>();
+			for (auto Obj : ReferencedObjects)
+			{
+				auto Field = Cast<UField>(Obj);
+				FString PackPath;
+				if (Field 
+					&& (Field->GetTypedOuter<UPackage>() != CurrentPackage)
+					&& FSourceCodeNavigation::FindClassHeaderPath(Field, PackPath))
+				{
+					if (!PackPath.RemoveFromStart(EngineSourceDir))
+					{
+						if (!PackPath.RemoveFromStart(GameSourceDir))
+						{
+							PackPath = FPaths::GetCleanFilename(PackPath);
+						}
+					}
+					HeaderFiles.Add(PackPath);
+				}
+			}
+		}
+
+		FString Result;
+		for (auto HeaderFile : HeaderFiles)
+		{
+			Result += FString::Printf(TEXT("#include \"%s\"\n"), *HeaderFile);
+		}
+
+		return Result;
+	}
+
+	
 };
 
 void FKismetCppBackend::EmitClassProperties(FStringOutputDevice& Target, UClass* SourceClass)
@@ -719,15 +787,15 @@ void FKismetCppBackend::EmitClassProperties(FStringOutputDevice& Target, UClass*
 	}
 }
 
-void FKismetCppBackend::GenerateCodeFromClass(UClass* SourceClass, TIndirectArray<FKismetFunctionContext>& Functions, bool bGenerateStubsOnly)
+void FKismetCppBackend::GenerateCodeFromClass(UClass* SourceClass, FString NewClassName, TIndirectArray<FKismetFunctionContext>& Functions, bool bGenerateStubsOnly)
 {
-	CppClassName = FEmitHelper::MakeCppClassName(SourceClass);
+	auto CleanCppClassName = NewClassName.IsEmpty() ? SourceClass->GetName() : NewClassName;
+	CppClassName = FString(SourceClass->GetPrefixCPP()) + CleanCppClassName;
 
 	UClass* SuperClass = SourceClass->GetSuperClass();
 
 	Emit(Header, TEXT("#pragma once\n\n"));
-	//Emit(Header, TEXT("#inlcude \"Public/Engine.h\"\n"));
-	Emit(Header, *FString::Printf(TEXT("#include \"%s.generated.h\"\n\n"), *SourceClass->GetName()));
+	Emit(Header, *FString::Printf(TEXT("#include \"%s.generated.h\"\n\n"), *CleanCppClassName));
 
 	// MC DELEGATE DECLARATION
 	{
@@ -794,6 +862,8 @@ void FKismetCppBackend::GenerateCodeFromClass(UClass* SourceClass, TIndirectArra
 	Emit(Body, *FString::Printf(TEXT("#include \"%s.h\"\n"), FApp::GetGameName()));
 	Emit(Body, TEXT("#include \"GeneratedCodeHelpers.h\"\n\n"));
 
+	Emit(Body, *FEmitHelper::GatherHeadersToInclude(SourceClass));
+
 	//constructor
 	Emit(Body, *FString::Printf(
 		TEXT("%s::%s(const FObjectInitializer& ObjectInitializer) : Super(ObjectInitializer) {}\n\n"),
@@ -809,7 +879,7 @@ void FKismetCppBackend::GenerateCodeFromClass(UClass* SourceClass, TIndirectArra
 
 	Emit(Header, TEXT("};\n\n"));
 
-	Emit(Body, *FEmitHelper::EmitLifetimeReplicatedPropsImpl(SourceClass, TEXT("")));
+	Emit(Body, *FEmitHelper::EmitLifetimeReplicatedPropsImpl(SourceClass, CppClassName, TEXT("")));
 }
 
 void FKismetCppBackend::EmitCallDelegateStatment(FKismetFunctionContext& FunctionContext, FBlueprintCompiledStatement& Statement)
@@ -1320,11 +1390,7 @@ void FKismetCppBackend::ConstructFunction(FKismetFunctionContext& FunctionContex
 		{
 			const FString Start = FString::Printf(TEXT("%s %s%s%s("), *ReturnType, TEXT("%s"), TEXT("%s"), *FunctionName);
 
-			const bool bUserConstructionScript = (FunctionName == TEXT("UserConstructionScript"));
-			if (!bUserConstructionScript)
-			{
-				Emit(Header, *FString::Printf(TEXT("\t%s\n"), *FEmitHelper::EmitUFuntion(Function)));
-			}
+			Emit(Header, *FString::Printf(TEXT("\t%s\n"), *FEmitHelper::EmitUFuntion(Function)));
 			Emit(Header, TEXT("\t"));
 			Emit(Header, *FString::Printf(*Start, TEXT(""), TEXT("")));
 			Emit(Body, *FString::Printf(*Start, *CppClassName, TEXT("::")));
@@ -1350,12 +1416,7 @@ void FKismetCppBackend::ConstructFunction(FKismetFunctionContext& FunctionContex
 			}
 
 			Emit(Header, TEXT(")"));
-			if (bUserConstructionScript)
-			{
-				Emit(Header, TEXT(" override"));
-			}
 			Emit(Header, TEXT(";\n"));
-			
 			Emit(Body, TEXT(")\n"));
 		}
 
diff --git a/Engine/Source/Editor/KismetCompiler/Private/KismetCompilerVMBackend.cpp b/Engine/Source/Editor/KismetCompiler/Private/KismetCompilerVMBackend.cpp
--- a/Engine/Source/Editor/KismetCompiler/Private/KismetCompilerVMBackend.cpp
+++ b/Engine/Source/Editor/KismetCompiler/Private/KismetCompilerVMBackend.cpp
@@ -685,7 +685,7 @@ public:
 
 #if UE_BLUEPRINT_EVENTGRAPH_FASTCALLS
 			// Store optimization data if this is a simple call into the ubergraph
-			if (FunctionContext.bIsSimpleStubGraphWithNoParams)
+			if (FunctionContext.bIsSimpleStubGraphWithNoParams && CompilerContext.NewClass->UberGraphFunction)
 			{
 				check(FunctionToCall == CompilerContext.NewClass->UberGraphFunction);
 				check(FunctionToCall->ParmsSize == sizeof(int32));
diff --git a/Engine/Source/Editor/UnrealEd/Private/Kismet2/Kismet2.cpp b/Engine/Source/Editor/UnrealEd/Private/Kismet2/Kismet2.cpp
--- a/Engine/Source/Editor/UnrealEd/Private/Kismet2/Kismet2.cpp
+++ b/Engine/Source/Editor/UnrealEd/Private/Kismet2/Kismet2.cpp
@@ -909,7 +909,7 @@ void FKismetEditorUtilities::RecompileBlueprintBytecode(UBlueprint* BlueprintObj
 }
 
 /** Recompiles the bytecode of a blueprint only.  Should only be run for recompiling dependencies during compile on load */
-void FKismetEditorUtilities::GenerateCppCode(UBlueprint* InBlueprintObj, TSharedPtr<FString> OutHeaderSource, TSharedPtr<FString> OutCppSource)
+void FKismetEditorUtilities::GenerateCppCode(UBlueprint* InBlueprintObj, TSharedPtr<FString> OutHeaderSource, TSharedPtr<FString> OutCppSource, const FString& OptionalClassName)
 {
 	check(InBlueprintObj);
 	check(InBlueprintObj->GetOutermost() != GetTransientPackage());
@@ -921,6 +921,8 @@ void FKismetEditorUtilities::GenerateCppCode(UBlueprint* InBlueprintObj, TShared
 	{
 		auto BlueprintObj = DuplicateObject<UBlueprint>(InBlueprintObj, GetTransientPackage(), *InBlueprintObj->GetName());
 		{
+			auto Reinstancer = FBlueprintCompileReinstancer::Create(BlueprintObj->GeneratedClass);
+
 			IKismetCompilerInterface& Compiler = FModuleManager::LoadModuleChecked<IKismetCompilerInterface>(KISMET_COMPILER_MODULENAME);
 
 			TGuardValue<bool> GuardTemplateNameFlag(GCompilingBlueprint, true);
@@ -930,6 +932,7 @@ void FKismetEditorUtilities::GenerateCppCode(UBlueprint* InBlueprintObj, TShared
 			CompileOptions.CompileType = EKismetCompileType::Cpp;
 			CompileOptions.OutCppSourceCode = OutCppSource;
 			CompileOptions.OutHeaderSourceCode = OutHeaderSource;
+			CompileOptions.NewCppClassName = OptionalClassName;
 			Compiler.CompileBlueprint(BlueprintObj, CompileOptions, Results);
 		}
 		BlueprintObj->RemoveGeneratedClasses();
diff --git a/Engine/Source/Editor/UnrealEd/Public/Kismet2/KismetEditorUtilities.h b/Engine/Source/Editor/UnrealEd/Public/Kismet2/KismetEditorUtilities.h
--- a/Engine/Source/Editor/UnrealEd/Public/Kismet2/KismetEditorUtilities.h
+++ b/Engine/Source/Editor/UnrealEd/Public/Kismet2/KismetEditorUtilities.h
@@ -78,7 +78,7 @@ public:
 	/** Recompiles the bytecode of a blueprint only.  Should only be run for recompiling dependencies during compile on load */
 	static void RecompileBlueprintBytecode(UBlueprint* BlueprintObj, TArray<UObject*>* ObjLoaded = NULL);
 
-	static void GenerateCppCode(UBlueprint* BlueprintObj, TSharedPtr<FString> OutHeaderSource, TSharedPtr<FString> OutCppSource);
+	static void GenerateCppCode(UBlueprint* BlueprintObj, TSharedPtr<FString> OutHeaderSource, TSharedPtr<FString> OutCppSource, const FString& OptionalClassName = FString());
 
 	/** Tries to make sure that a data-only blueprint is conformed to its native parent, in case any native class flags have changed */
 	static void ConformBlueprintFlagsAndComponents(UBlueprint* BlueprintObj);
diff --git a/Engine/Source/Runtime/Engine/Classes/Engine/Blueprint.h b/Engine/Source/Runtime/Engine/Classes/Engine/Blueprint.h
--- a/Engine/Source/Runtime/Engine/Classes/Engine/Blueprint.h
+++ b/Engine/Source/Runtime/Engine/Classes/Engine/Blueprint.h
@@ -81,6 +81,7 @@ public:
 
 	TSharedPtr<FString> OutHeaderSourceCode;
 	TSharedPtr<FString> OutCppSourceCode;
+	FString NewCppClassName;
 
 	bool DoesRequireCppCodeGeneration() const
 	{
