Optimizing various uses of Pop to prevent incremental reallocation for short-lived or frequently used TArray instances
[INTEGRATE] Change 2490222 by [REDACTED_PERSON][REDACTED_HANDLE] on 2015/03/25 00:18:46

diff --git a/Engine/Plugins/2D/Paper2D/Source/Paper2D/Private/PaperGeomTools.cpp b/Engine/Plugins/2D/Paper2D/Source/Paper2D/Private/PaperGeomTools.cpp
--- a/Engine/Plugins/2D/Paper2D/Source/Paper2D/Private/PaperGeomTools.cpp
+++ b/Engine/Plugins/2D/Paper2D/Source/Paper2D/Private/PaperGeomTools.cpp
@@ -664,7 +664,7 @@ void PaperGeomTools::RemoveRedundantTriangles(TArray<FVector2D>& OutTriangles, c
 	{
 		TArray<FVector2D> PolygonVertices;
 
-		const FLocalTriangle InitialTriangle = Triangles.Pop();
+		const FLocalTriangle InitialTriangle = Triangles.Pop(/*bAllowShrinking=*/ false);
 		PolygonVertices.Add(InTriangleVertices[InitialTriangle.VertexA]);
 		PolygonVertices.Add(InTriangleVertices[InitialTriangle.VertexB]);
 		PolygonVertices.Add(InTriangleVertices[InitialTriangle.VertexC]);
@@ -721,7 +721,7 @@ void PaperGeomTools::GenerateConvexPolygonsFromTriangles(TArray<TArray<FVector2D
 	{
 		TArray<FVector2D> PolygonVertices;
 
-		const FLocalTriangle InitialTriangle = Triangles.Pop();
+		const FLocalTriangle InitialTriangle = Triangles.Pop(/*bAllowShrinking=*/ false);
 		PolygonVertices.Add(InTriangleVertices[InitialTriangle.VertexA]);
 		PolygonVertices.Add(InTriangleVertices[InitialTriangle.VertexB]);
 		PolygonVertices.Add(InTriangleVertices[InitialTriangle.VertexC]);
diff --git a/Engine/Plugins/2D/Paper2D/Source/Paper2DEditor/Private/TileMapEditing/EdModeTileMap.cpp b/Engine/Plugins/2D/Paper2D/Source/Paper2DEditor/Private/TileMapEditing/EdModeTileMap.cpp
--- a/Engine/Plugins/2D/Paper2D/Source/Paper2DEditor/Private/TileMapEditing/EdModeTileMap.cpp
+++ b/Engine/Plugins/2D/Paper2D/Source/Paper2DEditor/Private/TileMapEditing/EdModeTileMap.cpp
@@ -742,7 +742,7 @@ bool FEdModeTileMap::FloodFillTiles(const FViewportCursorLocation& Ray)
 		// Process the list of outstanding spans until it is empty
 		while (OutstandingSpans.Num())
 		{
-			FHorizontalSpan Span = OutstandingSpans.Pop();
+			FHorizontalSpan Span = OutstandingSpans.Pop(/*bAllowShrinking=*/ false);
 
 			// Create spans below and above
 			for (int32 DY = -1; DY <= 1; DY += 2)
diff --git a/Engine/Source/Editor/KismetCompiler/Private/KismetCompilerCppBackend.cpp b/Engine/Source/Editor/KismetCompiler/Private/KismetCompilerCppBackend.cpp
--- a/Engine/Source/Editor/KismetCompiler/Private/KismetCompilerCppBackend.cpp
+++ b/Engine/Source/Editor/KismetCompiler/Private/KismetCompilerCppBackend.cpp
@@ -1246,7 +1246,7 @@ void FKismetCppBackend::EmitGotoStatement(FKismetFunctionContext& FunctionContex
 		if (Statement.Type == KCST_EndOfThreadIfNot)
 		{
 			ensure(FunctionContext.bUseFlowStack);
-			Emit(Body, TEXT("\t\t\t\tCurrentState = (StateStack.Num() > 0) ? StateStack.Pop() : -1;\n"));
+			Emit(Body, TEXT("\t\t\t\tCurrentState = (StateStack.Num() > 0) ? StateStack.Pop(/*bAllowShrinking=*/ false) : -1;\n"));
 		}
 		else if (Statement.Type == KCST_GotoReturnIfNot)
 		{
@@ -1280,7 +1280,7 @@ void FKismetCppBackend::EmitPushStateStatement(FKismetFunctionContext& FunctionC
 void FKismetCppBackend::EmitEndOfThreadStatement(FKismetFunctionContext& FunctionContext, const FString& ReturnValueString)
 {
 	ensure(FunctionContext.bUseFlowStack);
-	Emit(Body, TEXT("\t\t\tCurrentState = (StateStack.Num() > 0) ? StateStack.Pop() : -1;\n"));
+	Emit(Body, TEXT("\t\t\tCurrentState = (StateStack.Num() > 0) ? StateStack.Pop(/*bAllowShrinking=*/ false) : -1;\n"));
 	Emit(Body, TEXT("\t\t\tbreak;\n"));
 }
 
diff --git a/Engine/Source/Runtime/CoreUObject/Private/UObject/FindStronglyConnected.cpp b/Engine/Source/Runtime/CoreUObject/Private/UObject/FindStronglyConnected.cpp
--- a/Engine/Source/Runtime/CoreUObject/Private/UObject/FindStronglyConnected.cpp
+++ b/Engine/Source/Runtime/CoreUObject/Private/UObject/FindStronglyConnected.cpp
@@ -120,7 +120,7 @@ bool FFindStronglyConnected::FindSimpleCycleForComponentInner( TArray<UObject*>&
 		{
 			while (1)
 			{
-				UObject* Out = Stack.Pop();
+				UObject* Out = Stack.Pop(/*bAllowShrinking=*/ false);
 				Dest.Add(Out);
 				if (Out == Other)
 				{
@@ -181,7 +181,7 @@ FFindStronglyConnected::NodeInfo* FFindStronglyConnected::StrongConnectInner( UO
 		TArray<UObject*>& Dest = Components[Components.Add(TArray<UObject*>())];
 		while (1)
 		{
-			UObject* Out = Stack.Pop();
+			UObject* Out = Stack.Pop(/*bAllowShrinking=*/ false);
 			NodeInfo* OutVal = NodeIndex.Find(Out);
 			OutVal->InStack = false;
 			Dest.Add(Out);
diff --git a/Engine/Source/Runtime/CoreUObject/Private/UObject/ScriptCore.cpp b/Engine/Source/Runtime/CoreUObject/Private/UObject/ScriptCore.cpp
--- a/Engine/Source/Runtime/CoreUObject/Private/UObject/ScriptCore.cpp
+++ b/Engine/Source/Runtime/CoreUObject/Private/UObject/ScriptCore.cpp
@@ -1226,7 +1226,7 @@ void UObject::execPopExecutionFlow( FFrame& Stack, RESULT_DECL )
 	// Try to pop an entry off the stack and go there
 	if (Stack.FlowStack.Num())
 	{
-		CodeSkipSizeType Offset = Stack.FlowStack.Pop();
+		CodeSkipSizeType Offset = Stack.FlowStack.Pop(/*bAllowShrinking=*/ false);
 		Stack.Code = &Stack.Node->Script[ Offset ];
 	}
 	else
@@ -1251,7 +1251,7 @@ void UObject::execPopExecutionFlowIfNot( FFrame& Stack, RESULT_DECL )
 		// Try to pop an entry off the stack and go there
 		if (Stack.FlowStack.Num())
 		{
-			CodeSkipSizeType Offset = Stack.FlowStack.Pop();
+			CodeSkipSizeType Offset = Stack.FlowStack.Pop(/*bAllowShrinking=*/ false);
 			Stack.Code = &Stack.Node->Script[ Offset ];
 		}
 		else
diff --git a/Engine/Source/Runtime/Engine/Private/Components/SceneComponent.cpp b/Engine/Source/Runtime/Engine/Private/Components/SceneComponent.cpp
--- a/Engine/Source/Runtime/Engine/Private/Components/SceneComponent.cpp
+++ b/Engine/Source/Runtime/Engine/Private/Components/SceneComponent.cpp
@@ -1778,7 +1778,7 @@ void USceneComponent::SetVisibility(bool bNewVisibility, bool bPropagateToChildr
 
 		while (ComponentStack.Num() > 0)
 		{
-			USceneComponent* const CurrentComp = ComponentStack.Pop();
+			USceneComponent* const CurrentComp = ComponentStack.Pop(/*bAllowShrinking=*/ false);
 			if (CurrentComp)
 			{
 				ComponentStack.Append(CurrentComp->AttachChildren);
@@ -1811,7 +1811,7 @@ void USceneComponent::SetHiddenInGame(bool NewHiddenGame, bool bPropagateToChild
 
 		while (ComponentStack.Num() > 0)
 		{
-			USceneComponent* const CurrentComp = ComponentStack.Pop();
+			USceneComponent* const CurrentComp = ComponentStack.Pop(/*bAllowShrinking=*/ false);
 			if (CurrentComp)
 			{
 				ComponentStack.Append(CurrentComp->AttachChildren);
diff --git a/Engine/Source/Runtime/Engine/Private/EdGraph/EdGraph.cpp b/Engine/Source/Runtime/Engine/Private/EdGraph/EdGraph.cpp
--- a/Engine/Source/Runtime/Engine/Private/EdGraph/EdGraph.cpp
+++ b/Engine/Source/Runtime/Engine/Private/EdGraph/EdGraph.cpp
@@ -196,9 +196,10 @@ void UEdGraph::NotifyGraphChanged(const FEdGraphEditAction& InAction)
 void UEdGraph::MoveNodesToAnotherGraph(UEdGraph* DestinationGraph, bool bIsLoading, bool bInIsCompiling/* = false*/)
 {
 	// Move one node over at a time
+	DestinationGraph->Nodes.Reserve(DestinationGraph->Nodes.Num() + Nodes.Num());
 	while (Nodes.Num())
 	{
-		if (UEdGraphNode* Node = Nodes.Pop())
+		if (UEdGraphNode* Node = Nodes.Pop(/*bAllowShrinking=*/ false))
 		{
 #if WITH_EDITOR
 			// During compilation, do not move ghost nodes, they are not used during compilation.
diff --git a/Engine/Source/Runtime/Renderer/Private/SceneOcclusion.cpp b/Engine/Source/Runtime/Renderer/Private/SceneOcclusion.cpp
--- a/Engine/Source/Runtime/Renderer/Private/SceneOcclusion.cpp
+++ b/Engine/Source/Runtime/Renderer/Private/SceneOcclusion.cpp
@@ -67,7 +67,7 @@ FRenderQueryRHIRef FRenderQueryPool::AllocateQuery()
 		return RHICreateRenderQuery(QueryType);
 	}
 
-	return Queries.Pop();
+	return Queries.Pop(/*bAllowShrinking=*/ false);
 }
 
 void FRenderQueryPool::ReleaseQuery(FRHICommandListImmediate& RHICmdList, FRenderQueryRHIRef &Query)
diff --git a/Engine/Source/Runtime/Serialization/Private/StructDeserializer.cpp b/Engine/Source/Runtime/Serialization/Private/StructDeserializer.cpp
--- a/Engine/Source/Runtime/Serialization/Private/StructDeserializer.cpp
+++ b/Engine/Source/Runtime/Serialization/Private/StructDeserializer.cpp
@@ -109,7 +109,7 @@ bool FStructDeserializer::Deserialize( void* OutStruct, UStruct& TypeInfo, IStru
 					return false;
 				}
 
-				CurrentState = StateStack.Pop();
+				CurrentState = StateStack.Pop(/*bAllowShrinking=*/ false);
 			}
 			break;
 
@@ -220,7 +220,7 @@ bool FStructDeserializer::Deserialize( void* OutStruct, UStruct& TypeInfo, IStru
 					return true;
 				}
 
-				CurrentState = StateStack.Pop();
+				CurrentState = StateStack.Pop(/*bAllowShrinking=*/ false);
 			}
 			break;
 
diff --git a/Engine/Source/Runtime/Serialization/Private/StructSerializer.cpp b/Engine/Source/Runtime/Serialization/Private/StructSerializer.cpp
--- a/Engine/Source/Runtime/Serialization/Private/StructSerializer.cpp
+++ b/Engine/Source/Runtime/Serialization/Private/StructSerializer.cpp
@@ -90,7 +90,7 @@ void FStructSerializer::Serialize( const void* Struct, UStruct& TypeInfo, IStruc
 	// process state stack
 	while (StateStack.Num() > 0)
 	{
-		FWriteState CurrentState = StateStack.Pop();
+		FWriteState CurrentState = StateStack.Pop(/*bAllowShrinking=*/ false);
 
 		// structures
 		if ((CurrentState.Property == nullptr) || (CurrentState.TypeInfo == UStructProperty::StaticClass()))
diff --git a/Engine/Source/Runtime/SlateCore/Private/Rendering/ElementBatcher.cpp b/Engine/Source/Runtime/SlateCore/Private/Rendering/ElementBatcher.cpp
--- a/Engine/Source/Runtime/SlateCore/Private/Rendering/ElementBatcher.cpp
+++ b/Engine/Source/Runtime/SlateCore/Private/Rendering/ElementBatcher.cpp
@@ -1463,7 +1463,7 @@ FSlateElementBatch& FSlateElementBatcher::FindBatchForElement(
 		// Get a free vertex array
 		if( VertexArrayFreeList.Num() > 0 )
 		{
-			ElementBatch->VertexArrayIndex = VertexArrayFreeList.Pop();
+			ElementBatch->VertexArrayIndex = VertexArrayFreeList.Pop(/*bAllowShrinking=*/ false);
 			BatchVertexArrays[ElementBatch->VertexArrayIndex].Reserve(200);		
 		}
 		else
@@ -1478,7 +1478,7 @@ FSlateElementBatch& FSlateElementBatcher::FindBatchForElement(
 		// Get a free index array
 		if( IndexArrayFreeList.Num() > 0 )
 		{
-			ElementBatch->IndexArrayIndex = IndexArrayFreeList.Pop();
+			ElementBatch->IndexArrayIndex = IndexArrayFreeList.Pop(/*bAllowShrinking=*/ false);
 			BatchIndexArrays[ElementBatch->IndexArrayIndex].Reserve(200);		
 		}
 		else
diff --git a/Engine/Source/Runtime/SlateRHIRenderer/Private/SlateRHIResourceManager.cpp b/Engine/Source/Runtime/SlateRHIRenderer/Private/SlateRHIResourceManager.cpp
--- a/Engine/Source/Runtime/SlateRHIRenderer/Private/SlateRHIResourceManager.cpp
+++ b/Engine/Source/Runtime/SlateRHIRenderer/Private/SlateRHIResourceManager.cpp
@@ -414,7 +414,7 @@ TSharedPtr<FSlateDynamicTextureResource> FSlateRHIResourceManager::MakeDynamicTe
 	// Get a resource from the free list if possible
 	if(DynamicTextureFreeList.Num() > 0)
 	{
-		TextureResource = DynamicTextureFreeList.Pop();
+		TextureResource = DynamicTextureFreeList.Pop(/*bAllowShrinking=*/ false);
 	}
 	else
 	{
@@ -482,7 +482,7 @@ TSharedPtr<FSlateUTextureResource> FSlateRHIResourceManager::MakeDynamicUTexture
 		// Get a resource from the free list if possible
 		if (UTextureFreeList.Num() > 0)
 		{
-			TextureResource = UTextureFreeList.Pop(); 
+			TextureResource = UTextureFreeList.Pop(/*bAllowShrinking=*/ false);
 			TextureResource->TextureObject = InTextureObject;
 		}
 		else
