Default values in UDS does not override values in BPGC.

STRUCT_REINST_ uses the same GUID as the original struct.

diff --git a/Engine/Source/Editor/KismetCompiler/Private/UserDefinedStructureCompilerUtils.cpp b/Engine/Source/Editor/KismetCompiler/Private/UserDefinedStructureCompilerUtils.cpp
--- a/Engine/Source/Editor/KismetCompiler/Private/UserDefinedStructureCompilerUtils.cpp
+++ b/Engine/Source/Editor/KismetCompiler/Private/UserDefinedStructureCompilerUtils.cpp
@@ -26,6 +26,8 @@ struct FUserDefinedStructureCompilerInner
 				TGuardValue<bool> IsDuplicatingClassForReinstancing(GIsDuplicatingClassForReinstancing, true);
 				DuplicatedStruct = (UUserDefinedStruct*)StaticDuplicateObject(StructureToReinstance, GetTransientPackage(), *UniqueName.ToString(), ~RF_Transactional); 
 			}
+
+			DuplicatedStruct->Guid = StructureToReinstance->Guid;
 			DuplicatedStruct->Bind();
 			DuplicatedStruct->StaticLink(true);
 			DuplicatedStruct->PrimaryStruct = StructureToReinstance;
