Fixing an issue with not properly re-linking child skeleton classes after re-instancing due to no child classes being found.

diff --git a/Engine/Source/Editor/KismetCompiler/Private/KismetCompilerModule.cpp b/Engine/Source/Editor/KismetCompiler/Private/KismetCompilerModule.cpp
--- a/Engine/Source/Editor/KismetCompiler/Private/KismetCompilerModule.cpp
+++ b/Engine/Source/Editor/KismetCompiler/Private/KismetCompilerModule.cpp
@@ -119,6 +119,10 @@ void FKismet2CompilerModule::CompileBlueprint(class UBlueprint* Blueprint, const
 	{
 		BP_SCOPED_COMPILER_EVENT_STAT(EKismetCompilerStats_CompileSkeletonClass);
 
+		// Pre-cache the derived classes, before reinstancing can occur.
+		TArray<UClass*> ChildrenOfClass;
+		GetDerivedClasses(Blueprint->SkeletonGeneratedClass, ChildrenOfClass);
+
 		auto SkeletonReinstancer = FBlueprintCompileReinstancer::Create(Blueprint->SkeletonGeneratedClass);
 
 		FCompilerResultsLog SkeletonResults;
@@ -128,8 +132,6 @@ void FKismet2CompilerModule::CompileBlueprint(class UBlueprint* Blueprint, const
 		CompileBlueprintInner(Blueprint, SkeletonCompileOptions, SkeletonResults, ObjLoaded);
 
 		// Relink the child classes
-		TArray<UClass*> ChildrenOfClass;
-		GetDerivedClasses(Blueprint->SkeletonGeneratedClass, ChildrenOfClass);
 		for ( UClass* ChildClass : ChildrenOfClass )
 		{
 			ChildClass->AssembleReferenceTokenStream();
