UE-12003 Crash Reparenting Blueprint

Super class of Skel class isn't changed without reinstancing (the proper layout is necessary for CDO).
FBlueprintCompileReinstancer cannot be entered twice - the could happen after a compiler error.

diff --git a/Engine/Source/Editor/KismetCompiler/Private/KismetCompilerModule.cpp b/Engine/Source/Editor/KismetCompiler/Private/KismetCompilerModule.cpp
--- a/Engine/Source/Editor/KismetCompiler/Private/KismetCompilerModule.cpp
+++ b/Engine/Source/Editor/KismetCompiler/Private/KismetCompilerModule.cpp
@@ -130,24 +130,6 @@ void FKismet2CompilerModule::CompileBlueprint(class UBlueprint* Blueprint, const
 		FKismetCompilerOptions SkeletonCompileOptions;
 		SkeletonCompileOptions.CompileType = EKismetCompileType::SkeletonOnly;
 		CompileBlueprintInner(Blueprint, SkeletonCompileOptions, SkeletonResults, ObjLoaded);
-
-		// Relink the child classes
-		for ( UClass* ChildClass : ChildrenOfClass )
-		{
-			ChildClass->AssembleReferenceTokenStream();
-			ChildClass->SetSuperStruct(Blueprint->SkeletonGeneratedClass);
-
-			auto BPGDuplicatedClass = Cast<UBlueprintGeneratedClass>(ChildClass);
-			auto DuplicatedClassUberGraphFunction = BPGDuplicatedClass ? BPGDuplicatedClass->UberGraphFunction : nullptr;
-			if (DuplicatedClassUberGraphFunction)
-			{
-				DuplicatedClassUberGraphFunction->Bind();
-				DuplicatedClassUberGraphFunction->StaticLink(true);
-			}
-
-			ChildClass->Bind();
-			ChildClass->StaticLink(true);
-		}
 	}
 
 	// If this was a full compile, take appropriate actions depending on the success of failure of the compile
diff --git a/Engine/Source/Editor/UnrealEd/Private/Kismet2/KismetReinstanceUtilities.cpp b/Engine/Source/Editor/UnrealEd/Private/Kismet2/KismetReinstanceUtilities.cpp
--- a/Engine/Source/Editor/UnrealEd/Private/Kismet2/KismetReinstanceUtilities.cpp
+++ b/Engine/Source/Editor/UnrealEd/Private/Kismet2/KismetReinstanceUtilities.cpp
@@ -421,7 +421,7 @@ void FBlueprintCompileReinstancer::ReinstanceObjects(bool bForceAlwaysReinstance
 		static TArray<TSharedRef<FBlueprintCompileReinstancer>> QueueToReinstance;
 		TSharedRef<FBlueprintCompileReinstancer> SharedThis = AsShared();
 		const bool bAlreadyQueued = QueueToReinstance.Contains(SharedThis);
-		if (!bAlreadyQueued)
+		if (!bAlreadyQueued && !bHasReinstanced)
 		{
 			QueueToReinstance.Push(SharedThis);
 
@@ -456,6 +456,7 @@ void FBlueprintCompileReinstancer::ReinstanceObjects(bool bForceAlwaysReinstance
 				for (int32 Idx = 0; Idx < QueueToReinstance.Num(); ++Idx)
 				{
 					QueueToReinstance[Idx]->ReinstanceInner(bForceAlwaysReinstance);
+					QueueToReinstance[Idx]->bHasReinstanced = true;
 				}
 				QueueToReinstance.Empty();
 			}
