Sanitized widgets and their variables' names in the UMGEditor
[UE-9480] - UMG: Default names for Custom Widget variables in the Hierarchy and My Blueprint have ugly unsanitized names

#change Always tries to generate friendly labels for widgets in the UMG hierarchy panel (still always shown as the real object name in the variables list on the BP)
#change UWidgets created by adding items to the hierarchy panel are given sanitised object names when generated from BP classes

diff --git a/Engine/Source/Editor/KismetCompiler/Private/KismetCompilerMisc.cpp b/Engine/Source/Editor/KismetCompiler/Private/KismetCompilerMisc.cpp
--- a/Engine/Source/Editor/KismetCompiler/Private/KismetCompilerMisc.cpp
+++ b/Engine/Source/Editor/KismetCompiler/Private/KismetCompilerMisc.cpp
@@ -727,19 +727,21 @@ UProperty* FKismetCompilerUtilities::CreatePropertyOnScope(UStruct* Scope, const
 			return nullptr;
 		};
 
-		if (UObject* ExistingObject = CheckIfPropertyNameIsUsed(Scope, *PropertyName.ToString()))
+		if (UObject* ExistingObject = CheckPropertyNameOnScope(Scope, PropertyName))
 		{
 			MessageLog.Error(*FString::Printf(TEXT("Internal Compiler Error:  Tried to create a property %s in scope %s, but %s already exists there."), *PropertyName.ToString(), (Scope ? *Scope->GetName() : TEXT("None")), *ExistingObject->GetFullName()));
 
 			// Find a free name, so we can still create the property to make it easier to spot the duplicates, and avoid crashing
 			uint32 Counter = 0;
-			FString TestNameString;
+			FName TestName;
 			do 
 			{
-				TestNameString = PropertyName.ToString() + FString::Printf(TEXT("_ERROR_DUPLICATE_%d"), Counter++);
-			} while (CheckIfPropertyNameIsUsed(Scope, *TestNameString) != NULL);
+				FString TestNameString = PropertyName.ToString() + FString::Printf(TEXT("_ERROR_DUPLICATE_%d"), Counter++);
+				TestName = FName(*TestNameString);
 
-			ValidatedPropertyName = FName(*TestNameString);
+			} while (CheckPropertyNameOnScope(Scope, TestName) != NULL);
+
+			ValidatedPropertyName = TestName;
 		}
 	}
 
@@ -921,6 +923,25 @@ UProperty* FKismetCompilerUtilities::CreatePropertyOnScope(UStruct* Scope, const
 	return NewProperty;
 }
 
+UObject* FKismetCompilerUtilities::CheckPropertyNameOnScope(UStruct* Scope, const FName& PropertyName)
+{
+	FString NameStr = PropertyName.ToString();
+
+	if (UObject* ExistingObject = FindObject<UObject>(Scope, *NameStr, false))
+	{
+		return ExistingObject;
+	}
+
+	if (Scope && !Scope->IsA<UFunction>() && (UBlueprintGeneratedClass::GetUberGraphFrameName() != PropertyName))
+	{
+		if (auto Field = FindField<UProperty>(Scope ? Scope->GetSuperStruct() : nullptr, *NameStr))
+		{
+			return Field;
+		}
+	}
+
+	return nullptr;
+}
 
 //////////////////////////////////////////////////////////////////////////
 // FNodeHandlingFunctor
diff --git a/Engine/Source/Editor/KismetCompiler/Public/KismetCompilerMisc.h b/Engine/Source/Editor/KismetCompiler/Public/KismetCompilerMisc.h
--- a/Engine/Source/Editor/KismetCompiler/Public/KismetCompilerMisc.h
+++ b/Engine/Source/Editor/KismetCompiler/Public/KismetCompilerMisc.h
@@ -63,6 +63,13 @@ public:
 	/** Creates a property named PropertyName of type PropertyType in the Scope or returns NULL if the type is unknown, but does *not* link that property in */
 	static UProperty* CreatePropertyOnScope(UStruct* Scope, const FName& PropertyName, const FEdGraphPinType& Type, UClass* SelfClass, uint64 PropertyFlags, const class UEdGraphSchema_K2* Schema, FCompilerResultsLog& MessageLog);
 
+	/**
+	 * Checks that the property name isn't taken in the given scope (used by CreatePropertyOnScope())
+	 *
+	 * @return	Ptr to an existing object with that name in the given scope or nullptr if none exists
+	 */
+	static UObject* CheckPropertyNameOnScope(UStruct* Scope, const FName& PropertyName);
+
 private:
 	/** Counter to ensure unique names in the transient package, to avoid GC collection issues with classes and their CDOs */
 	static uint32 ConsignToOblivionCounter;
diff --git a/Engine/Source/Editor/UMGEditor/Private/Designer/SDesignerView.cpp b/Engine/Source/Editor/UMGEditor/Private/Designer/SDesignerView.cpp
--- a/Engine/Source/Editor/UMGEditor/Private/Designer/SDesignerView.cpp
+++ b/Engine/Source/Editor/UMGEditor/Private/Designer/SDesignerView.cpp
@@ -93,8 +93,8 @@ TSharedRef<FSelectedWidgetDragDropOp> FSelectedWidgetDragDropOp::New(TSharedPtr<
 	TSharedRef<FSelectedWidgetDragDropOp> Operation = MakeShareable(new FSelectedWidgetDragDropOp());
 	Operation->bStayingInParent = bStayInParent;
 	Operation->ParentWidget = Editor->GetReferenceFromTemplate(InWidget.GetTemplate()->GetParent());
-	Operation->DefaultHoverText = FText::FromString( InWidget.GetTemplate()->GetLabel() );
-	Operation->CurrentHoverText = FText::FromString( InWidget.GetTemplate()->GetLabel() );
+	Operation->DefaultHoverText = InWidget.GetTemplate()->GetLabelText();
+	Operation->CurrentHoverText = InWidget.GetTemplate()->GetLabelText();
 	Operation->Construct();
 
 	// Cache the preview and template, it's not safe to query the preview/template while dragging the widget as it no longer
diff --git a/Engine/Source/Editor/UMGEditor/Private/Hierarchy/SHierarchyViewItem.cpp b/Engine/Source/Editor/UMGEditor/Private/Hierarchy/SHierarchyViewItem.cpp
--- a/Engine/Source/Editor/UMGEditor/Private/Hierarchy/SHierarchyViewItem.cpp
+++ b/Engine/Source/Editor/UMGEditor/Private/Hierarchy/SHierarchyViewItem.cpp
@@ -60,8 +60,8 @@ TSharedRef<FHierarchyWidgetDragDropOp> FHierarchyWidgetDragDropOp::New(UWidgetBl
 {
 	TSharedRef<FHierarchyWidgetDragDropOp> Operation = MakeShareable(new FHierarchyWidgetDragDropOp());
 	Operation->Widget = InWidget;
-	Operation->DefaultHoverText = FText::FromString(InWidget.GetTemplate()->GetLabel());
-	Operation->CurrentHoverText = FText::FromString(InWidget.GetTemplate()->GetLabel());
+	Operation->DefaultHoverText = InWidget.GetTemplate()->GetLabelText();
+	Operation->CurrentHoverText = InWidget.GetTemplate()->GetLabelText();
 	Operation->Construct();
 
 	FWidgetBlueprintEditorUtils::ExportPropertiesToText(InWidget.GetTemplate()->Slot, Operation->ExportedSlotProperties);
@@ -646,7 +646,7 @@ FText FHierarchyWidget::GetText() const
 	UWidget* WidgetTemplate = Item.GetTemplate();
 	if ( WidgetTemplate )
 	{
-		return FText::FromString(WidgetTemplate->GetLabel());
+		return WidgetTemplate->GetLabelText();
 	}
 
 	return FText::GetEmpty();
@@ -679,7 +679,7 @@ FText FHierarchyWidget::GetLabelToolTipText() const
 	UWidget* WidgetTemplate = Item.GetTemplate();
 	if ( WidgetTemplate && !WidgetTemplate->IsGeneratedName() )
 	{
-		return FText::FromString(TEXT( "[" ) + WidgetTemplate->GetClass()->GetName() + TEXT( "]" ) );
+		return FText::FromString(TEXT( "[" ) + WidgetTemplate->GetClass()->GetDisplayNameText().ToString() + TEXT( "]" ) );
 	}
 
 	return FText::GetEmpty();
diff --git a/Engine/Source/Editor/UMGEditor/Private/Templates/WidgetTemplateBlueprintClass.cpp b/Engine/Source/Editor/UMGEditor/Private/Templates/WidgetTemplateBlueprintClass.cpp
--- a/Engine/Source/Editor/UMGEditor/Private/Templates/WidgetTemplateBlueprintClass.cpp
+++ b/Engine/Source/Editor/UMGEditor/Private/Templates/WidgetTemplateBlueprintClass.cpp
@@ -5,6 +5,7 @@
 #include "WidgetTemplateBlueprintClass.h"
 #include "IDocumentation.h"
 #include "WidgetBlueprint.h"
+#include "Kismet2/BlueprintEditorUtils.h"
 
 #define LOCTEXT_NAMESPACE "UMGEditor"
 
@@ -60,7 +61,7 @@ UWidget* FWidgetTemplateBlueprintClass::Create(UWidgetTree* Tree)
 		WidgetClass = CastChecked<UClass>(LoadedWidget->GeneratedClass);
 	}
 
-	return FWidgetTemplateClass::Create(Tree);
+	return FWidgetTemplateClass::CreateNamed(Tree, FName(*FBlueprintEditorUtils::GetClassNameWithoutSuffix(WidgetClass.Get())));
 }
 
 const FSlateBrush* FWidgetTemplateBlueprintClass::GetIcon() const
diff --git a/Engine/Source/Editor/UMGEditor/Private/Templates/WidgetTemplateClass.cpp b/Engine/Source/Editor/UMGEditor/Private/Templates/WidgetTemplateClass.cpp
--- a/Engine/Source/Editor/UMGEditor/Private/Templates/WidgetTemplateClass.cpp
+++ b/Engine/Source/Editor/UMGEditor/Private/Templates/WidgetTemplateClass.cpp
@@ -37,10 +37,7 @@ FText FWidgetTemplateClass::GetCategory() const
 
 UWidget* FWidgetTemplateClass::Create(UWidgetTree* Tree)
 {
-	UWidget* NewWidget = Tree->ConstructWidget<UWidget>(WidgetClass.Get());
-	NewWidget->OnCreationFromPalette();
-
-	return NewWidget;
+	return CreateNamed(Tree, NAME_None);
 }
 
 const FSlateBrush* FWidgetTemplateClass::GetIcon() const
@@ -63,4 +60,21 @@ void FWidgetTemplateClass::OnObjectsReplaced(const TMap<UObject*, UObject*>& Rep
 	}
 }
 
+UWidget* FWidgetTemplateClass::CreateNamed(class UWidgetTree* Tree, FName NameOverride)
+{
+	if (NameOverride != NAME_None)
+	{
+		UObject* ExistingObject = StaticFindObject(UObject::StaticClass(), Tree, *NameOverride.ToString());
+		if (ExistingObject != nullptr)
+		{
+			NameOverride = MakeUniqueObjectName(Tree, WidgetClass.Get(), NameOverride);
+		}
+	}
+
+	UWidget* NewWidget = Tree->ConstructWidget<UWidget>(WidgetClass.Get(), NameOverride);
+	NewWidget->OnCreationFromPalette();
+
+	return NewWidget;
+}
+
 #undef LOCTEXT_NAMESPACE
diff --git a/Engine/Source/Editor/UMGEditor/Private/Templates/WidgetTemplateClass.h b/Engine/Source/Editor/UMGEditor/Private/Templates/WidgetTemplateClass.h
--- a/Engine/Source/Editor/UMGEditor/Private/Templates/WidgetTemplateClass.h
+++ b/Engine/Source/Editor/UMGEditor/Private/Templates/WidgetTemplateClass.h
@@ -37,6 +37,10 @@ protected:
 	/** Called when objects need to be swapped out for new versions, like after a blueprint recompile. */
 	void OnObjectsReplaced(const TMap<UObject*, UObject*>& ReplacementMap);
 
+	/** Constructs the widget template with an overridden object name. */
+	UWidget* CreateNamed(class UWidgetTree* Tree, FName NameOverride);
+
+protected:
 	/** The widget class that will be created by this template */
 	TWeakObjectPtr<UClass> WidgetClass;
 };
diff --git a/Engine/Source/Runtime/UMG/Private/Components/Widget.cpp b/Engine/Source/Runtime/UMG/Private/Components/Widget.cpp
--- a/Engine/Source/Runtime/UMG/Private/Components/Widget.cpp
+++ b/Engine/Source/Runtime/UMG/Private/Components/Widget.cpp
@@ -526,14 +526,27 @@ FString UWidget::GetLabelMetadata() const
 
 FString UWidget::GetLabel() const
 {
-	if ( IsGeneratedName() && !bIsVariable )
-	{
-		return TEXT("[") + GetClass()->GetName() + TEXT("]") + GetLabelMetadata();
-	}
-	else
+	return GetLabelText().ToString();
+}
+
+FText UWidget::GetLabelText() const
+{
+	FText Label = GetDisplayNameBase();
+
+	if (IsGeneratedName() && !bIsVariable)
 	{
-		return GetName();
+		FFormatNamedArguments Args;
+		Args.Add(TEXT("BaseName"), Label);
+		Args.Add(TEXT("Metadata"), FText::FromString(GetLabelMetadata()));
+		Label = FText::Format(LOCTEXT("NonVariableLabelFormat", "[{BaseName}]{Metadata}"), Args);
 	}
+
+	return Label;
+}
+
+FText UWidget::GetDisplayNameBase() const
+{
+	return IsGeneratedName() ? GetClass()->GetDisplayNameText() : FText::FromString(GetName());
 }
 
 const FText UWidget::GetPaletteCategory()
diff --git a/Engine/Source/Runtime/UMG/Public/Blueprint/WidgetTree.h b/Engine/Source/Runtime/UMG/Public/Blueprint/WidgetTree.h
--- a/Engine/Source/Runtime/UMG/Public/Blueprint/WidgetTree.h
+++ b/Engine/Source/Runtime/UMG/Public/Blueprint/WidgetTree.h
@@ -102,18 +102,18 @@ public:
 
 	/** Constructs the widget, and adds it to the tree. */
 	template< class T >
-	FORCEINLINE T* ConstructWidget(TSubclassOf<UWidget> WidgetType)
+	FORCEINLINE T* ConstructWidget(TSubclassOf<UWidget> WidgetType, FName WidgetName = NAME_None)
 	{
 		if ( WidgetType->IsChildOf(UUserWidget::StaticClass()) )
 		{
-			UUserWidget* Widget = NewObject<UUserWidget>(this, WidgetType);
+			UUserWidget* Widget = NewObject<UUserWidget>(this, WidgetType, WidgetName);
 			Widget->Initialize();
 			Widget->SetFlags(RF_Transactional);
 			return (T*)Widget;
 		}
 		else
 		{
-			UWidget* Widget = (UWidget*)NewObject<UWidget>(this, WidgetType);
+			UWidget* Widget = (UWidget*)NewObject<UWidget>(this, WidgetType, WidgetName);
 			Widget->SetFlags(RF_Transactional);
 			return (T*)Widget;
 		}
diff --git a/Engine/Source/Runtime/UMG/Public/Components/Widget.h b/Engine/Source/Runtime/UMG/Public/Components/Widget.h
--- a/Engine/Source/Runtime/UMG/Public/Components/Widget.h
+++ b/Engine/Source/Runtime/UMG/Public/Components/Widget.h
@@ -389,8 +389,12 @@ public:
 	virtual FString GetLabelMetadata() const;
 
 	/** Gets the label to display to the user for this widget. */
+	DEPRECATED(4.8, "Use GetLabelText(), which will return the label as FText.")
 	FString GetLabel() const;
 
+	/** Gets the label to display to the user for this widget. */
+	FText GetLabelText() const;
+
 	/** Gets the palette category of the widget */
 	virtual const FText GetPaletteCategory();
 
@@ -451,6 +455,9 @@ protected:
 
 	void UpdateRenderTransform();
 
+	/** Gets the base name used to generate the display label/name of this widget. */
+	FText GetDisplayNameBase() const;
+
 protected:
 	//TODO UMG Consider moving conversion functions into another class.
 	// Conversion functions
