New BP compiler warning:  "The execution path doesn't end with a return node."

diff --git a/Engine/Source/Editor/KismetCompiler/Private/KismetCompiler.cpp b/Engine/Source/Editor/KismetCompiler/Private/KismetCompiler.cpp
--- a/Engine/Source/Editor/KismetCompiler/Private/KismetCompiler.cpp
+++ b/Engine/Source/Editor/KismetCompiler/Private/KismetCompiler.cpp
@@ -1222,6 +1222,8 @@ void FKismetCompilerContext::PrecompileFunction(FKismetFunctionContext& Context)
 		//@TODO: Prune pure functions that don't have any consumers
 		if (bIsFullCompile)
 		{
+			FKismetCompilerUtilities::ValidateProperEndExecutionPath(Context);
+
 			// Find the execution path (and make sure it has no cycles)
 			CreateExecutionSchedule(Context.SourceGraph->Nodes, Context.LinearExecutionList);
 
diff --git a/Engine/Source/Editor/KismetCompiler/Private/KismetCompilerMisc.cpp b/Engine/Source/Editor/KismetCompiler/Private/KismetCompilerMisc.cpp
--- a/Engine/Source/Editor/KismetCompiler/Private/KismetCompilerMisc.cpp
+++ b/Engine/Source/Editor/KismetCompiler/Private/KismetCompilerMisc.cpp
@@ -709,24 +709,6 @@ UProperty* FKismetCompilerUtilities::CreatePropertyOnScope(UStruct* Scope, const
 	// Check to see if there's already a object on this scope with the same name, and throw an internal compiler error if so
 	// If this happens, it breaks the property link, which causes stack corruption and hard-to-track errors, so better to fail at this point
 	{
-		auto CheckIfPropertyNameIsUsed = [](UStruct* Struct, const TCHAR* Name) -> UObject*
-		{
-			if (UObject* ExistingObject = FindObject<UObject>(Struct, Name, false))
-			{
-				return ExistingObject;
-			}
-			
-			if (Struct && !Struct->IsA<UFunction>() && (UBlueprintGeneratedClass::GetUberGraphFrameName() != Name))
-			{
-				if (auto Field = FindField<UProperty>(Struct ? Struct->GetSuperStruct() : nullptr, Name))
-				{
-					return Field;
-				}
-			}
-
-			return nullptr;
-		};
-
 		if (UObject* ExistingObject = CheckPropertyNameOnScope(Scope, PropertyName))
 		{
 			MessageLog.Error(*FString::Printf(TEXT("Internal Compiler Error:  Tried to create a property %s in scope %s, but %s already exists there."), *PropertyName.ToString(), (Scope ? *Scope->GetName() : TEXT("None")), *ExistingObject->GetFullName()));
@@ -943,6 +925,67 @@ UObject* FKismetCompilerUtilities::CheckPropertyNameOnScope(UStruct* Scope, cons
 	return nullptr;
 }
 
+/** Checks if the execution path ends with a Return node */
+void FKismetCompilerUtilities::ValidateProperEndExecutionPath(FKismetFunctionContext& Context)
+{
+	struct FRecrursiveHelper
+	{
+		// returns if the path is properly ended
+		static void CheckPath(UK2Node* StartingNode, TSet<UK2Node*>& VisitedNodes, FKismetFunctionContext& Context)
+		{
+			UK2Node* CurrentNode = StartingNode;
+			while (CurrentNode)
+			{
+				UK2Node* SourceNode = CurrentNode;
+				CurrentNode = nullptr;
+
+				bool bAlreadyVisited = false;
+				VisitedNodes.Add(SourceNode, &bAlreadyVisited);
+				if (!bAlreadyVisited && !SourceNode->IsA<UK2Node_FunctionResult>())
+				{
+					const bool bIsExecutionSequence = SourceNode->IsA<UK2Node_ExecutionSequence>();
+					for (auto CurrentPin : SourceNode->Pins)
+					{
+						if (CurrentPin
+							&& (CurrentPin->Direction == EEdGraphPinDirection::EGPD_Output)
+							&& (CurrentPin->PinType.PinCategory == Context.Schema->PC_Exec)
+							&& (CurrentPin->LinkedTo.Num() > 0))
+						{
+							auto LinkedPin = CurrentPin->LinkedTo[0];
+							auto NextNode = ensure(LinkedPin) ? Cast<UK2Node>(LinkedPin->GetOwningNodeUnchecked()) : nullptr;
+							if (CurrentNode && !bIsExecutionSequence && ensure(NextNode)) // for sequence node, we want only the last output
+							{
+								FRecrursiveHelper::CheckPath(NextNode, VisitedNodes, Context);
+							}
+							else
+							{
+								CurrentNode = NextNode;
+							}
+						}
+					}
+
+					if (!CurrentNode)
+					{
+						Context.MessageLog.Warning(*LOCTEXT("ExecutionEnd_Warning", "The execution path doesn't end with a return node. @@").ToString(), SourceNode);
+					}
+				}
+			}
+		}
+	};
+
+	// Function is designed for multiple return nodes.
+	if (!Context.IsEventGraph() && Context.SourceGraph && Context.Schema)
+	{
+		TArray<UK2Node_FunctionResult*> ReturnNodes;
+		Context.SourceGraph->GetNodesOfClass(ReturnNodes);
+		if (ReturnNodes.Num() && ensure(Context.EntryPoint))
+		{
+			TSet<UK2Node*> VisitedNodes;
+			FRecrursiveHelper::CheckPath(Context.EntryPoint, VisitedNodes, Context);
+		}
+	}
+}
+
 //////////////////////////////////////////////////////////////////////////
 // FNodeHandlingFunctor
 
diff --git a/Engine/Source/Editor/KismetCompiler/Public/KismetCompilerMisc.h b/Engine/Source/Editor/KismetCompiler/Public/KismetCompilerMisc.h
--- a/Engine/Source/Editor/KismetCompiler/Public/KismetCompilerMisc.h
+++ b/Engine/Source/Editor/KismetCompiler/Public/KismetCompilerMisc.h
@@ -91,6 +91,9 @@ public:
 
 	/** Create Kismet assignment statement with proper object <-> interface cast */
 	static void CreateObjectAssignmentStatement(FKismetFunctionContext& Context, UEdGraphNode* Node, FBPTerminal* SrcTerm, FBPTerminal* DstTerm);
+
+	/** Checks if each execution path ends with a Return node */
+	static void ValidateProperEndExecutionPath(FKismetFunctionContext& Context);
 };
 
 //////////////////////////////////////////////////////////////////////////
