Proper "The execution path doesn't end with a return node." warning for MultiGate node.

diff --git a/Engine/Source/Editor/KismetCompiler/Private/KismetCompilerMisc.cpp b/Engine/Source/Editor/KismetCompiler/Private/KismetCompilerMisc.cpp
--- a/Engine/Source/Editor/KismetCompiler/Private/KismetCompilerMisc.cpp
+++ b/Engine/Source/Editor/KismetCompiler/Private/KismetCompilerMisc.cpp
@@ -943,7 +943,7 @@ void FKismetCompilerUtilities::ValidateProperEndExecutionPath(FKismetFunctionCon
 				VisitedNodes.Add(SourceNode, &bAlreadyVisited);
 				if (!bAlreadyVisited && !SourceNode->IsA<UK2Node_FunctionResult>())
 				{
-					const bool bIsExecutionSequence = SourceNode->IsA<UK2Node_ExecutionSequence>();
+					const bool bIsExecutionSequence = UK2Node_ExecutionSequence::StaticClass() == SourceNode->GetClass(); // no "SourceNode->IsA<UK2Node_ExecutionSequence>()" because MultiGate is based on ExecutionSequence
 					for (auto CurrentPin : SourceNode->Pins)
 					{
 						if (CurrentPin
