UE-7848 BP nodes are executed after the return node is reached

Function result node generates KCST_GotoReturn statement.
"The execution path doesn't end with a return node." Is now a note.

diff --git a/Engine/Source/Editor/BlueprintGraph/Private/K2Node_FunctionResult.cpp b/Engine/Source/Editor/BlueprintGraph/Private/K2Node_FunctionResult.cpp
--- a/Engine/Source/Editor/BlueprintGraph/Private/K2Node_FunctionResult.cpp
+++ b/Engine/Source/Editor/BlueprintGraph/Private/K2Node_FunctionResult.cpp
@@ -65,6 +65,22 @@ public:
 		Term->CopyFromPin(Net, Net->PinName);
 		Context.NetMap.Add(Net, Term);
 	}
+
+	virtual void Compile(FKismetFunctionContext& Context, UEdGraphNode* Node) override
+	{
+		GenerateAssigments(Context, Node);
+
+		if (Context.bCreateDebugData && Node)
+		{
+			FBlueprintCompiledStatement& TraceStatement = Context.AppendStatementForNode(Node);
+			TraceStatement.Type = KCST_WireTraceSite;
+			TraceStatement.Comment = Node->NodeComment.IsEmpty() ? Node->GetName() : Node->NodeComment;
+		}
+
+		// always go to return
+		FBlueprintCompiledStatement& GotoStatement = Context.AppendStatementForNode(Node);
+		GotoStatement.Type = KCST_GotoReturn; 
+	}
 };
 
 UK2Node_FunctionResult::UK2Node_FunctionResult(const FObjectInitializer& ObjectInitializer)
diff --git a/Engine/Source/Editor/BlueprintGraph/Private/VariableSetHandler.cpp b/Engine/Source/Editor/BlueprintGraph/Private/VariableSetHandler.cpp
--- a/Engine/Source/Editor/BlueprintGraph/Private/VariableSetHandler.cpp
+++ b/Engine/Source/Editor/BlueprintGraph/Private/VariableSetHandler.cpp
@@ -110,7 +110,7 @@ void FKCHandler_VariableSet::InnerAssignment(FKismetFunctionContext& Context, UE
 	}
 }
 
-void FKCHandler_VariableSet::Compile(FKismetFunctionContext& Context, UEdGraphNode* Node)
+void FKCHandler_VariableSet::GenerateAssigments(FKismetFunctionContext& Context, UEdGraphNode* Node)
 {
 	// SubCategory is an object type or "" for the stack frame, default scope is Self
 	// Each input pin is the name of a variable
@@ -132,6 +132,11 @@ void FKCHandler_VariableSet::Compile(FKismetFunctionContext& Context, UEdGraphNo
 			CompilerContext.MessageLog.Error(*FString::Printf(*LOCTEXT("ExpectedOnlyInputPins_Error", "Expected only input pins on @@ but found @@").ToString()), Node, Pin);
 		}
 	}
+}
+
+void FKCHandler_VariableSet::Compile(FKismetFunctionContext& Context, UEdGraphNode* Node)
+{
+	GenerateAssigments(Context, Node);
 
 	// Generate the output impulse from this node
 	GenerateSimpleThenGoto(Context, *Node);
diff --git a/Engine/Source/Editor/BlueprintGraph/Public/VariableSetHandler.h b/Engine/Source/Editor/BlueprintGraph/Public/VariableSetHandler.h
--- a/Engine/Source/Editor/BlueprintGraph/Public/VariableSetHandler.h
+++ b/Engine/Source/Editor/BlueprintGraph/Public/VariableSetHandler.h
@@ -16,6 +16,7 @@ public:
 	virtual void RegisterNet(FKismetFunctionContext& Context, UEdGraphPin* Net) override;
 	virtual void RegisterNets(FKismetFunctionContext& Context, UEdGraphNode* Node) override;
 	void InnerAssignment(FKismetFunctionContext& Context, UEdGraphNode* Node, UEdGraphPin* VariablePin, UEdGraphPin* ValuePin);
+	void GenerateAssigments(FKismetFunctionContext& Context, UEdGraphNode* Node);
 	virtual void Compile(FKismetFunctionContext& Context, UEdGraphNode* Node) override;
 	virtual void Transform(FKismetFunctionContext& Context, UEdGraphNode* Node) override;
 };
diff --git a/Engine/Source/Editor/KismetCompiler/Private/KismetCompiler.cpp b/Engine/Source/Editor/KismetCompiler/Private/KismetCompiler.cpp
--- a/Engine/Source/Editor/KismetCompiler/Private/KismetCompiler.cpp
+++ b/Engine/Source/Editor/KismetCompiler/Private/KismetCompiler.cpp
@@ -51,7 +51,7 @@ DECLARE_CYCLE_STAT(TEXT("Clean and Sanitize Class"), EKismetCompilerStats_CleanA
 DECLARE_CYCLE_STAT(TEXT("Create Class Properties"), EKismetCompilerStats_CreateClassVariables, STATGROUP_KismetCompiler );
 DECLARE_CYCLE_STAT(TEXT("Bind and Link Class"), EKismetCompilerStats_BindAndLinkClass, STATGROUP_KismetCompiler );
 DECLARE_CYCLE_STAT(TEXT("Calculate checksum of CDO"), EKismetCompilerStats_ChecksumCDO, STATGROUP_KismetCompiler );
-		
+DECLARE_CYCLE_STAT(TEXT("Analyze execution path"), EKismetCompilerStats_AnalyzeExecutionPath, STATGROUP_KismetCompiler);
 //////////////////////////////////////////////////////////////////////////
 // FKismetCompilerContext
 
@@ -1226,7 +1226,11 @@ void FKismetCompilerContext::PrecompileFunction(FKismetFunctionContext& Context)
 		//@TODO: Prune pure functions that don't have any consumers
 		if (bIsFullCompile)
 		{
-			FKismetCompilerUtilities::ValidateProperEndExecutionPath(Context);
+			if (!Blueprint->bIsRegeneratingOnLoad)
+			{
+				BP_SCOPED_COMPILER_EVENT_STAT(EKismetCompilerStats_AnalyzeExecutionPath);
+				FKismetCompilerUtilities::ValidateProperEndExecutionPath(Context);
+			}
 
 			// Find the execution path (and make sure it has no cycles)
 			CreateExecutionSchedule(Context.SourceGraph->Nodes, Context.LinearExecutionList);
diff --git a/Engine/Source/Editor/KismetCompiler/Private/KismetCompilerMisc.cpp b/Engine/Source/Editor/KismetCompiler/Private/KismetCompilerMisc.cpp
--- a/Engine/Source/Editor/KismetCompiler/Private/KismetCompilerMisc.cpp
+++ b/Engine/Source/Editor/KismetCompiler/Private/KismetCompilerMisc.cpp
@@ -931,13 +931,12 @@ void FKismetCompilerUtilities::ValidateProperEndExecutionPath(FKismetFunctionCon
 {
 	struct FRecrursiveHelper
 	{
-		// returns if the path is properly ended
-		static void CheckPath(UK2Node* StartingNode, TSet<UK2Node*>& VisitedNodes, FKismetFunctionContext& InContext)
+		static void CheckPathEnding(const UK2Node* StartingNode, TSet<const UK2Node*>& VisitedNodes, FKismetFunctionContext& InContext, bool bPathShouldEndWithReturn)
 		{
-			UK2Node* CurrentNode = StartingNode;
+			const UK2Node* CurrentNode = StartingNode;
 			while (CurrentNode)
 			{
-				UK2Node* SourceNode = CurrentNode;
+				const UK2Node* SourceNode = CurrentNode;
 				CurrentNode = nullptr;
 
 				bool bAlreadyVisited = false;
@@ -949,26 +948,26 @@ void FKismetCompilerUtilities::ValidateProperEndExecutionPath(FKismetFunctionCon
 					{
 						if (CurrentPin
 							&& (CurrentPin->Direction == EEdGraphPinDirection::EGPD_Output)
-							&& (CurrentPin->PinType.PinCategory == InContext.Schema->PC_Exec)
-							&& (CurrentPin->LinkedTo.Num() > 0))
+							&& (CurrentPin->PinType.PinCategory == InContext.Schema->PC_Exec))
 						{
-							auto LinkedPin = CurrentPin->LinkedTo[0];
-							auto NextNode = ensure(LinkedPin) ? Cast<UK2Node>(LinkedPin->GetOwningNodeUnchecked()) : nullptr;
-							if (CurrentNode && !bIsExecutionSequence && ensure(NextNode)) // for sequence node, we want only the last output
+							if (!CurrentPin->LinkedTo.Num())
 							{
-								FRecrursiveHelper::CheckPath(NextNode, VisitedNodes, InContext);
+								if (bPathShouldEndWithReturn && !bIsExecutionSequence)
+								{
+									InContext.MessageLog.Note(*LOCTEXT("ExecutionEnd_Note", "The execution path doesn't end with a return node. @@").ToString(), CurrentPin);
+								}
+								continue;
 							}
-							else
+							auto LinkedPin = CurrentPin->LinkedTo[0];
+							auto NextNode = ensure(LinkedPin) ? Cast<const UK2Node>(LinkedPin->GetOwningNodeUnchecked()) : nullptr;
+							ensure(NextNode);
+							if (CurrentNode)
 							{
-								CurrentNode = NextNode;
+								FRecrursiveHelper::CheckPathEnding(CurrentNode, VisitedNodes, InContext, bPathShouldEndWithReturn && !bIsExecutionSequence);
 							}
+							CurrentNode = NextNode;
 						}
 					}
-
-					if (!CurrentNode)
-					{
-						InContext.MessageLog.Warning(*LOCTEXT("ExecutionEnd_Warning", "The execution path doesn't end with a return node. @@").ToString(), SourceNode);
-					}
 				}
 			}
 		}
@@ -981,8 +980,8 @@ void FKismetCompilerUtilities::ValidateProperEndExecutionPath(FKismetFunctionCon
 		Context.SourceGraph->GetNodesOfClass(ReturnNodes);
 		if (ReturnNodes.Num() && ensure(Context.EntryPoint))
 		{
-			TSet<UK2Node*> VisitedNodes;
-			FRecrursiveHelper::CheckPath(Context.EntryPoint, VisitedNodes, Context);
+			TSet<const UK2Node*> VisitedNodes;
+			FRecrursiveHelper::CheckPathEnding(Context.EntryPoint, VisitedNodes, Context, true);
 		}
 	}
 }
