BP: "The path is never executed." note added (for functions with a return node).

diff --git a/Engine/Source/Editor/KismetCompiler/Private/KismetCompilerMisc.cpp b/Engine/Source/Editor/KismetCompiler/Private/KismetCompilerMisc.cpp
--- a/Engine/Source/Editor/KismetCompiler/Private/KismetCompilerMisc.cpp
+++ b/Engine/Source/Editor/KismetCompiler/Private/KismetCompilerMisc.cpp
@@ -931,7 +931,12 @@ void FKismetCompilerUtilities::ValidateProperEndExecutionPath(FKismetFunctionCon
 {
 	struct FRecrursiveHelper
 	{
-		static void CheckPathEnding(const UK2Node* StartingNode, TSet<const UK2Node*>& VisitedNodes, FKismetFunctionContext& InContext, bool bPathShouldEndWithReturn)
+		static bool IsExecutionSequence(const UEdGraphNode* Node)
+		{
+			return Node && (UK2Node_ExecutionSequence::StaticClass() == Node->GetClass()); // no "SourceNode->IsA<UK2Node_ExecutionSequence>()" because MultiGate is based on ExecutionSequence
+		}
+
+		static void CheckPathEnding(const UK2Node* StartingNode, TSet<const UK2Node*>& VisitedNodes, FKismetFunctionContext& InContext, bool bPathShouldEndWithReturn, TSet<const UK2Node*>& BreakableNodesSeeds)
 		{
 			const UK2Node* CurrentNode = StartingNode;
 			while (CurrentNode)
@@ -943,7 +948,7 @@ void FKismetCompilerUtilities::ValidateProperEndExecutionPath(FKismetFunctionCon
 				VisitedNodes.Add(SourceNode, &bAlreadyVisited);
 				if (!bAlreadyVisited && !SourceNode->IsA<UK2Node_FunctionResult>())
 				{
-					const bool bIsExecutionSequence = UK2Node_ExecutionSequence::StaticClass() == SourceNode->GetClass(); // no "SourceNode->IsA<UK2Node_ExecutionSequence>()" because MultiGate is based on ExecutionSequence
+					const bool bIsExecutionSequence = IsExecutionSequence(SourceNode); 
 					for (auto CurrentPin : SourceNode->Pins)
 					{
 						if (CurrentPin
@@ -952,6 +957,10 @@ void FKismetCompilerUtilities::ValidateProperEndExecutionPath(FKismetFunctionCon
 						{
 							if (!CurrentPin->LinkedTo.Num())
 							{
+								if (!bIsExecutionSequence)
+								{
+									BreakableNodesSeeds.Add(SourceNode);
+								}
 								if (bPathShouldEndWithReturn && !bIsExecutionSequence)
 								{
 									InContext.MessageLog.Note(*LOCTEXT("ExecutionEnd_Note", "The execution path doesn't end with a return node. @@").ToString(), CurrentPin);
@@ -963,7 +972,7 @@ void FKismetCompilerUtilities::ValidateProperEndExecutionPath(FKismetFunctionCon
 							ensure(NextNode);
 							if (CurrentNode)
 							{
-								FRecrursiveHelper::CheckPathEnding(CurrentNode, VisitedNodes, InContext, bPathShouldEndWithReturn && !bIsExecutionSequence);
+								FRecrursiveHelper::CheckPathEnding(CurrentNode, VisitedNodes, InContext, bPathShouldEndWithReturn && !bIsExecutionSequence, BreakableNodesSeeds);
 							}
 							CurrentNode = NextNode;
 						}
@@ -971,6 +980,132 @@ void FKismetCompilerUtilities::ValidateProperEndExecutionPath(FKismetFunctionCon
 				}
 			}
 		}
+
+		static void GatherBreakableNodes(const UK2Node* StartingNode, TSet<const UK2Node*>& BreakableNodes, FKismetFunctionContext& InContext)
+		{
+			const UK2Node* CurrentNode = StartingNode;
+			while (CurrentNode)
+			{
+				const UK2Node* SourceNode = CurrentNode;
+				CurrentNode = nullptr;
+
+				bool bAlreadyVisited = false;
+				BreakableNodes.Add(SourceNode, &bAlreadyVisited);
+				if (!bAlreadyVisited)
+				{
+					for (auto CurrentPin : SourceNode->Pins)
+					{
+						if (CurrentPin
+							&& (CurrentPin->Direction == EEdGraphPinDirection::EGPD_Input)
+							&& (CurrentPin->PinType.PinCategory == InContext.Schema->PC_Exec)
+							&& CurrentPin->LinkedTo.Num())
+						{
+							for (auto LinkedPin : CurrentPin->LinkedTo)
+							{
+								auto NextNode = ensure(LinkedPin) ? Cast<const UK2Node>(LinkedPin->GetOwningNodeUnchecked()) : nullptr;
+								ensure(NextNode);
+								if (!FRecrursiveHelper::IsExecutionSequence(NextNode))
+								{
+									if (CurrentNode)
+									{
+										GatherBreakableNodes(CurrentNode, BreakableNodes, InContext);
+									}
+									CurrentNode = NextNode;
+								}
+							}
+						}
+					}
+				}
+			}
+		}
+
+		static void GatherBreakableNodesSeedsFromSequences(TSet<const UK2Node_ExecutionSequence*>& UnBreakableExecutionSequenceNodes
+			, TSet<const UK2Node*>& BreakableNodesSeeds
+			, TSet<const UK2Node*>& BreakableNodes
+			, FKismetFunctionContext& InContext)
+		{
+			for (auto SequenceNode : UnBreakableExecutionSequenceNodes)
+			{
+				bool bIsBreakable = true;
+				// Sequence is breakable when all it's outputs are breakable
+				for (auto CurrentPin : SequenceNode->Pins)
+				{
+					if (CurrentPin
+						&& (CurrentPin->Direction == EEdGraphPinDirection::EGPD_Output)
+						&& (CurrentPin->PinType.PinCategory == InContext.Schema->PC_Exec)
+						&& CurrentPin->LinkedTo.Num())
+					{
+						auto LinkedPin = CurrentPin->LinkedTo[0];
+						auto NextNode = ensure(LinkedPin) ? Cast<const UK2Node>(LinkedPin->GetOwningNodeUnchecked()) : nullptr;
+						ensure(NextNode);
+						if (!BreakableNodes.Contains(NextNode))
+						{
+							bIsBreakable = false;
+							break;
+						}
+					}
+				}
+
+				if (bIsBreakable)
+				{
+					bool bWasAlreadyBreakable = false;
+					BreakableNodesSeeds.Add(SequenceNode, &bWasAlreadyBreakable);
+					ensure(!bWasAlreadyBreakable);
+					int32 WasRemoved = UnBreakableExecutionSequenceNodes.Remove(SequenceNode);
+					ensure(WasRemoved);
+				}
+			}
+		}
+
+		static void CheckDeadExecutionPath(TSet<const UK2Node*>& BreakableNodesSeeds, FKismetFunctionContext& InContext)
+		{
+			TSet<const UK2Node_ExecutionSequence*> UnBreakableExecutionSequenceNodes;
+			for (auto Node : InContext.SourceGraph->Nodes)
+			{
+				if (FRecrursiveHelper::IsExecutionSequence(Node))
+				{
+					UnBreakableExecutionSequenceNodes.Add(Cast<UK2Node_ExecutionSequence>(Node));
+				}
+			}
+
+			TSet<const UK2Node*> BreakableNodes;
+			while (BreakableNodesSeeds.Num())
+			{
+				for (auto StartingNode : BreakableNodesSeeds)
+				{
+					GatherBreakableNodes(StartingNode, BreakableNodes, InContext);
+				}
+				BreakableNodesSeeds.Empty();
+				FRecrursiveHelper::GatherBreakableNodesSeedsFromSequences(UnBreakableExecutionSequenceNodes, BreakableNodesSeeds, BreakableNodes, InContext);
+			}
+
+			for (auto UnBreakableExecutionSequenceNode : UnBreakableExecutionSequenceNodes)
+			{
+				bool bUnBreakableOutputWasFound = false;
+				for (auto CurrentPin : UnBreakableExecutionSequenceNode->Pins)
+				{
+					if (CurrentPin
+						&& (CurrentPin->Direction == EEdGraphPinDirection::EGPD_Output)
+						&& (CurrentPin->PinType.PinCategory == InContext.Schema->PC_Exec)
+						&& CurrentPin->LinkedTo.Num())
+					{
+						if (bUnBreakableOutputWasFound)
+						{
+							InContext.MessageLog.Note(*LOCTEXT("DeadExecution_Note", "The path is never executed. @@").ToString(), CurrentPin);
+							break;
+						}
+
+						auto LinkedPin = CurrentPin->LinkedTo[0];
+						auto NextNode = ensure(LinkedPin) ? Cast<const UK2Node>(LinkedPin->GetOwningNodeUnchecked()) : nullptr;
+						ensure(NextNode);
+						if (!BreakableNodes.Contains(NextNode))
+						{
+							bUnBreakableOutputWasFound = true;
+						}
+					}
+				}
+			}
+		}
 	};
 
 	// Function is designed for multiple return nodes.
@@ -980,8 +1115,13 @@ void FKismetCompilerUtilities::ValidateProperEndExecutionPath(FKismetFunctionCon
 		Context.SourceGraph->GetNodesOfClass(ReturnNodes);
 		if (ReturnNodes.Num() && ensure(Context.EntryPoint))
 		{
-			TSet<const UK2Node*> VisitedNodes;
-			FRecrursiveHelper::CheckPathEnding(Context.EntryPoint, VisitedNodes, Context, true);
+			TSet<const UK2Node*> VisitedNodes, BreakableNodesSeeds;
+			FRecrursiveHelper::CheckPathEnding(Context.EntryPoint, VisitedNodes, Context, true, BreakableNodesSeeds);
+
+			// A non-pure node, that lies on a execution path, that may result with "EndThread" state, is called Breakable.
+			// The execution path between the node and Return node can be broken.
+
+			FRecrursiveHelper::CheckDeadExecutionPath(BreakableNodesSeeds, Context);
 		}
 	}
 }
