Automatically generated Array for macro input is empty.

diff --git a/Engine/Source/Editor/KismetCompiler/Private/KismetCompiler.cpp b/Engine/Source/Editor/KismetCompiler/Private/KismetCompiler.cpp
--- a/Engine/Source/Editor/KismetCompiler/Private/KismetCompiler.cpp
+++ b/Engine/Source/Editor/KismetCompiler/Private/KismetCompiler.cpp
@@ -2852,6 +2852,7 @@ void FKismetCompilerContext::ExpandTunnelsAndMacros(UEdGraph* SourceGraph)
 					&& (Pin->LinkedTo.Num() == 0))
 					{
 						UK2Node_MakeArray* MakeArrayNode = SpawnIntermediateNode<UK2Node_MakeArray>(MacroInstanceNode, SourceGraph);
+						MakeArrayNode->NumInputs = 0; // the generated array should be empty
 						MakeArrayNode->AllocateDefaultPins();
 						UEdGraphPin* MakeArrayOut = MakeArrayNode->GetOutputPin();
 						check(MakeArrayOut);
