[UE-12238] Exposed 'const' attribute on user-defined functions in order to be able to call them from 'const' BPIE/BPNE overrides.

change summary:
- exposed 'const' attribute for user-defined functions as an advanced option.
- function flags are now set during the precompile phase rather than the post-compile phase when compiling functions during skeleton/full BP compiler passes. this ensures that function validation works correctly during a function graph's compile phase if the function graph calls another function that's defined within the same Blueprint class.

diff --git a/Engine/Source/Editor/Kismet/Private/BlueprintDetailsCustomization.cpp b/Engine/Source/Editor/Kismet/Private/BlueprintDetailsCustomization.cpp
--- a/Engine/Source/Editor/Kismet/Private/BlueprintDetailsCustomization.cpp
+++ b/Engine/Source/Editor/Kismet/Private/BlueprintDetailsCustomization.cpp
@@ -2452,6 +2452,23 @@ void FBlueprintGraphActionDetails::CustomizeDetails( IDetailLayoutBuilder& Detai
 						.OnCheckStateChanged( this, &FBlueprintGraphActionDetails::OnIsPureFunctionModified )
 				];
 			}
+			if (IsConstFunctionVisible())
+			{
+				Category.AddCustomRow( LOCTEXT( "FunctionConst_Tooltip", "Const" ), true )
+				.NameContent()
+				[
+					SNew(STextBlock)
+					.Text( LOCTEXT( "FunctionConst_Tooltip", "Const" ) )
+					.ToolTipText( LOCTEXT("FunctionIsConst_Tooltip", "Force this to be a const function?") )
+					.Font( IDetailLayoutBuilder::GetDetailFont() )
+				]
+				.ValueContent()
+				[
+					SNew( SCheckBox )
+					.IsChecked( this, &FBlueprintGraphActionDetails::GetIsConstFunction )
+					.OnCheckStateChanged( this, &FBlueprintGraphActionDetails::OnIsConstFunctionModified )
+				];
+			}
 		}
 
 		if (IsCustomEvent())
@@ -3732,6 +3749,50 @@ ECheckBoxState FBlueprintGraphActionDetails::GetIsPureFunction() const
 	return (EntryNode->ExtraFlags & FUNC_BlueprintPure) ? ECheckBoxState::Checked :  ECheckBoxState::Unchecked;
 }
 
+bool FBlueprintGraphActionDetails::IsConstFunctionVisible() const
+{
+	bool bSupportedType = false;
+	bool bIsEditable = false;
+	UK2Node_EditablePinBase * FunctionEntryNode = FunctionEntryNodePtr.Get();
+	if(FunctionEntryNode)
+	{
+		UBlueprint* Blueprint = FunctionEntryNode->GetBlueprint();
+
+		bSupportedType = FunctionEntryNode->IsA<UK2Node_FunctionEntry>();
+		bIsEditable = FunctionEntryNode->IsEditable();
+	}
+	return bSupportedType && bIsEditable;
+}
+
+void FBlueprintGraphActionDetails::OnIsConstFunctionModified( const ECheckBoxState NewCheckedState )
+{
+	UK2Node_EditablePinBase * FunctionEntryNode = FunctionEntryNodePtr.Get();
+	auto Function = FindFunction();
+	auto EntryNode = Cast<UK2Node_FunctionEntry>(FunctionEntryNode);
+	if(EntryNode && Function)
+	{
+		const FScopedTransaction Transaction( LOCTEXT( "ChangeConst", "Change Const" ) );
+		EntryNode->Modify();
+		Function->Modify();
+
+		//set flags on function entry node also
+		EntryNode->ExtraFlags	^= FUNC_Const;
+		Function->FunctionFlags ^= FUNC_Const;
+		OnParamsChanged(FunctionEntryNode);
+	}
+}
+
+ECheckBoxState FBlueprintGraphActionDetails::GetIsConstFunction() const
+{
+	UK2Node_EditablePinBase * FunctionEntryNode = FunctionEntryNodePtr.Get();
+	auto EntryNode = Cast<UK2Node_FunctionEntry>(FunctionEntryNode);
+	if(!EntryNode)
+	{
+		return ECheckBoxState::Undetermined;
+	}
+	return (EntryNode->ExtraFlags & FUNC_Const) ? ECheckBoxState::Checked :  ECheckBoxState::Unchecked;
+}
+
 FReply FBaseBlueprintGraphActionDetails::OnAddNewInputClicked()
 {
 	UK2Node_EditablePinBase * FunctionEntryNode = FunctionEntryNodePtr.Get();
diff --git a/Engine/Source/Editor/Kismet/Private/BlueprintDetailsCustomization.h b/Engine/Source/Editor/Kismet/Private/BlueprintDetailsCustomization.h
--- a/Engine/Source/Editor/Kismet/Private/BlueprintDetailsCustomization.h
+++ b/Engine/Source/Editor/Kismet/Private/BlueprintDetailsCustomization.h
@@ -514,6 +514,10 @@ private:
 	void OnIsPureFunctionModified(const ECheckBoxState NewCheckedState);
 	ECheckBoxState GetIsPureFunction() const;
 
+	bool IsConstFunctionVisible() const;
+	void OnIsConstFunctionModified(const ECheckBoxState NewCheckedState);
+	ECheckBoxState GetIsConstFunction() const;
+
 	/** Determines if the selected event is identified as editor callable */
 	ECheckBoxState GetIsEditorCallableEvent() const;
 
diff --git a/Engine/Source/Editor/KismetCompiler/Private/KismetCompiler.cpp b/Engine/Source/Editor/KismetCompiler/Private/KismetCompiler.cpp
--- a/Engine/Source/Editor/KismetCompiler/Private/KismetCompiler.cpp
+++ b/Engine/Source/Editor/KismetCompiler/Private/KismetCompiler.cpp
@@ -1205,6 +1205,64 @@ void FKismetCompilerContext::PrecompileFunction(FKismetFunctionContext& Context)
 		{
 			Context.Function->SetMetaData(FBlueprintMetadata::MD_CallInEditor, TEXT( "true" ));
 		}
+
+		// Set the required function flags
+		if (Context.CanBeCalledByKismet())
+		{
+			Context.Function->FunctionFlags |= FUNC_BlueprintCallable;
+		}
+
+		if (Context.IsInterfaceStub())
+		{
+			Context.Function->FunctionFlags |= FUNC_BlueprintEvent;
+		}
+
+		// Inherit extra flags from the entry node
+		if (Context.EntryPoint)
+		{
+			Context.Function->FunctionFlags |= Context.EntryPoint->ExtraFlags;
+		}
+
+		// First try to get the overriden function from the super class
+		UFunction* OverridenFunction = Context.Function->GetSuperFunction();
+		// If we couldn't find it, see if we can find an interface class in our inheritance to get it from
+		if (!OverridenFunction && Context.Blueprint)
+		{
+			bool bInvalidInterface = false;
+			OverridenFunction = FBlueprintEditorUtils::FindFunctionInImplementedInterfaces( Context.Blueprint, Context.Function->GetFName(), &bInvalidInterface );
+			if(bInvalidInterface)
+			{
+				MessageLog.Warning(TEXT("Blueprint tried to implement invalid interface."));
+			}
+		}
+
+		// Inherit flags and validate against overridden function if it exists
+		if (OverridenFunction)
+		{
+			Context.Function->FunctionFlags |= (OverridenFunction->FunctionFlags & (FUNC_FuncInherit | FUNC_Public | FUNC_Protected | FUNC_Private));
+
+			if ((Context.Function->FunctionFlags & FUNC_AccessSpecifiers) != (OverridenFunction->FunctionFlags & FUNC_AccessSpecifiers))
+			{
+				MessageLog.Error(*LOCTEXT("IncompatibleAccessSpecifier_Error", "Access specifier is not compatible the parent function @@").ToString(), Context.EntryPoint);
+			}
+
+			const uint32 OverrideFlagsToCheck = (FUNC_FuncOverrideMatch & ~FUNC_AccessSpecifiers);
+			if ((Context.Function->FunctionFlags & OverrideFlagsToCheck) != (OverridenFunction->FunctionFlags & OverrideFlagsToCheck))
+			{
+				MessageLog.Error(*LOCTEXT("IncompatibleOverrideFlags_Error", "Overriden function is not compatible with the parent function @@. Check flags: Exec, Final, Static.").ToString(), Context.EntryPoint);
+			}
+
+			// Copy metadata from parent function as well
+			UMetaData::CopyMetadata(OverridenFunction, Context.Function);
+		}
+		else
+		{
+			// If this is the root of a blueprint-defined function or event, and if it's public, make it overrideable
+			if( !Context.IsEventGraph() && !Context.Function->HasAnyFunctionFlags(FUNC_Private) )
+			{
+				Context.Function->FunctionFlags |= FUNC_BlueprintEvent;
+			}
+		}
 		
 		// Link it
 		//@TODO: should this be in regular or reverse order?
@@ -1516,64 +1574,6 @@ void FKismetCompilerContext::FinishCompilingFunction(FKismetFunctionContext& Con
 	Function->Bind();
 	Function->StaticLink(true);
 
-	// Set the required function flags
-	if (Context.CanBeCalledByKismet())
-	{
-		Function->FunctionFlags |= FUNC_BlueprintCallable;
-	}
-
-	if (Context.IsInterfaceStub())
-	{
-		Function->FunctionFlags |= FUNC_BlueprintEvent;
-	}
-
-	// Inherit extra flags from the entry node
-	if (Context.EntryPoint)
-	{
-		Function->FunctionFlags |= Context.EntryPoint->ExtraFlags;
-	}
-
-	// First try to get the overriden function from the super class
-	UFunction* OverridenFunction = Function->GetSuperFunction();
-	// If we couldn't find it, see if we can find an interface class in our inheritance to get it from
-	if (!OverridenFunction && Context.Blueprint)
-	{
-		bool bInvalidInterface = false;
-		OverridenFunction = FBlueprintEditorUtils::FindFunctionInImplementedInterfaces( Context.Blueprint, Function->GetFName(), &bInvalidInterface );
-		if(bInvalidInterface)
-			{
-				MessageLog.Warning(TEXT("Blueprint tried to implement invalid interface."));
-			}
-		}
-
-	// Inherit flags and validate against overridden function if it exists
-	if (OverridenFunction)
-	{
-		Function->FunctionFlags |= (OverridenFunction->FunctionFlags & (FUNC_FuncInherit | FUNC_Public | FUNC_Protected | FUNC_Private));
-
-		if ((Function->FunctionFlags & FUNC_AccessSpecifiers) != (OverridenFunction->FunctionFlags & FUNC_AccessSpecifiers))
-		{
-			MessageLog.Error(*LOCTEXT("IncompatibleAccessSpecifier_Error", "Access specifier is not compatible the parent function @@").ToString(), Context.EntryPoint);
-		}
-
-		const uint32 OverrideFlagsToCheck = (FUNC_FuncOverrideMatch & ~FUNC_AccessSpecifiers);
-		if ((Function->FunctionFlags & OverrideFlagsToCheck) != (OverridenFunction->FunctionFlags & OverrideFlagsToCheck))
-		{
-			MessageLog.Error(*LOCTEXT("IncompatibleOverrideFlags_Error", "Overriden function is not compatible with the parent function @@. Check flags: Exec, Final, Static.").ToString(), Context.EntryPoint);
-		}
-
-		// Copy metadata from parent function as well
-		UMetaData::CopyMetadata(OverridenFunction, Function);
-	}
-	else
-	{
-		// If this is the root of a blueprint-defined function or event, and if it's public, make it overrideable
-		if( !Context.IsEventGraph() && !Function->HasAnyFunctionFlags(FUNC_Private) )
-		{
-			Function->FunctionFlags |= FUNC_BlueprintEvent;
-		}
-	}
-
 	// Set function flags and calculate cached values so the class can be used immediately
 	Function->ParmsSize = 0;
 	Function->NumParms = 0;
