Signature CRC32 for Blueprint

diff --git a/Engine/Source/Editor/KismetCompiler/Private/KismetCompiler.cpp b/Engine/Source/Editor/KismetCompiler/Private/KismetCompiler.cpp
--- a/Engine/Source/Editor/KismetCompiler/Private/KismetCompiler.cpp
+++ b/Engine/Source/Editor/KismetCompiler/Private/KismetCompiler.cpp
@@ -52,6 +52,7 @@ DECLARE_CYCLE_STAT(TEXT("Create Class Properties"), EKismetCompilerStats_CreateC
 DECLARE_CYCLE_STAT(TEXT("Bind and Link Class"), EKismetCompilerStats_BindAndLinkClass, STATGROUP_KismetCompiler );
 DECLARE_CYCLE_STAT(TEXT("Calculate checksum of CDO"), EKismetCompilerStats_ChecksumCDO, STATGROUP_KismetCompiler );
 DECLARE_CYCLE_STAT(TEXT("Analyze execution path"), EKismetCompilerStats_AnalyzeExecutionPath, STATGROUP_KismetCompiler);
+DECLARE_CYCLE_STAT(TEXT("Calculate checksum of signature"), EKismetCompilerStats_ChecksumSignature, STATGROUP_KismetCompiler);
 //////////////////////////////////////////////////////////////////////////
 // FKismetCompilerContext
 
@@ -3646,7 +3647,6 @@ void FKismetCompilerContext::Compile()
 			UPackage* const Package = Cast<UPackage>(CurrentBP->GetOutermost());
 			const bool bStartedWithUnsavedChanges = Package != nullptr ? Package->IsDirty() : true;
 
-			CurrentBP->Status = BS_Dirty;
 			FBlueprintEditorUtils::RefreshExternalBlueprintDependencyNodes(CurrentBP, NewClass);
 			
 			// Note: We do not send a change notification event to the dependent BP here because
@@ -3722,7 +3722,84 @@ void FKismetCompilerContext::Compile()
 
 		UObject* NewCDO = NewClass->GetDefaultObject(false);
 		FSpecializedArchiveCrc32 CrcArchive(!ChangeDefaultValueWithoutReinstancing);
-		Blueprint->CrcPreviousCompiledCDO = NewCDO ? CrcArchive.Crc32(NewCDO) : 0;
+		Blueprint->CrcLastCompiledCDO = NewCDO ? CrcArchive.Crc32(NewCDO) : 0;
+	}
+
+	if (bIsFullCompile)
+	{
+		BP_SCOPED_COMPILER_EVENT_STAT(EKismetCompilerStats_ChecksumSignature);
+
+		class FSignatureArchiveCrc32 : public FArchiveObjectCrc32
+		{
+		public:
+			static bool IsInnerProperty(const UObject* Object)
+			{
+				auto Property = Cast<const UProperty>(Object);
+				return Property // check arrays
+					&& Cast<const UFunction>(Property->GetOwnerStruct())
+					&& !Property->HasAnyPropertyFlags(CPF_Parm);
+			}
+
+			virtual FArchive& operator<<(UObject*& Object) override
+			{
+				FArchive& Ar = *this;
+
+				if (Object && !IsInnerProperty(Object))
+				{
+					// Names of functions and properties are significant.
+					auto UniqueName = GetPathNameSafe(Object);
+					Ar << UniqueName;
+
+					if (Object->IsIn(RootObject))
+					{
+						ObjectsToSerialize.Enqueue(Object);
+					}
+				}
+
+				return Ar;
+			}
+
+			virtual bool CustomSerialize(UObject* Object) override
+			{ 
+				FArchive& Ar = *this;
+
+				bool bResult = false;
+				if (auto Struct = Cast<UStruct>(Object))
+				{
+					if (Object == RootObject) // name and location are significant for the signature
+					{
+						auto UniqueName = GetPathNameSafe(Object);
+						Ar << UniqueName;
+					}
+
+					UObject* SuperStruct = Struct->GetSuperStruct();
+					Ar << SuperStruct;
+					Ar << Struct->Children;
+
+					if (auto Function = Cast<UFunction>(Struct))
+					{
+						Ar << Function->FunctionFlags;
+					}
+
+					if (auto AsClass = Cast<UClass>(Struct))
+					{
+						Ar << AsClass->ClassFlags;
+						Ar << AsClass->Interfaces;
+					}
+
+					Ar << Struct->Next;
+
+					bResult = true;
+				}
+
+				return bResult;
+			}
+		};
+
+		FSignatureArchiveCrc32 SignatureArchiveCrc32;
+		auto ParentBP = UBlueprint::GetBlueprintFromClass(NewClass->GetSuperClass());
+		const uint32 ParentSignatureCrc = ParentBP ? ParentBP->CrcLastCompiledSignature : 0;
+		Blueprint->CrcLastCompiledSignature = SignatureArchiveCrc32.Crc32(NewClass, ParentSignatureCrc);
 	}
 }
 
diff --git a/Engine/Source/Runtime/CoreUObject/Private/Serialization/ArchiveObjectCrc32.cpp b/Engine/Source/Runtime/CoreUObject/Private/Serialization/ArchiveObjectCrc32.cpp
--- a/Engine/Source/Runtime/CoreUObject/Private/Serialization/ArchiveObjectCrc32.cpp
+++ b/Engine/Source/Runtime/CoreUObject/Private/Serialization/ArchiveObjectCrc32.cpp
@@ -80,13 +80,16 @@ uint32 FArchiveObjectCrc32::Crc32(UObject* Object, uint32 CRC)
 #endif
 				// Serialize it
 				ObjectBeingSerialized = Object;
-				Object->Serialize(*this);
+				if (!CustomSerialize(Object))
+				{
+					Object->Serialize(*this);
+				}
 				ObjectBeingSerialized = NULL;
 
 				// Calculate the CRC, compounding it with the checksum calculated from the previous object
 				CRC = FCrc::MemCrc32(SerializedObjectData.GetData(), SerializedObjectData.Num(), CRC);
 #ifdef DEBUG_ARCHIVE_OBJECT_CRC32
-				UE_LOG(LogArchiveObjectCrc32, Log, TEXT("=> total size: %d bytes, checksum: 0x%08x"), SerializedObjectData.Num(), CRC);
+				UE_LOG(LogArchiveObjectCrc32, Log, TEXT("=> object: '%s', total size: %d bytes, checksum: 0x%08x"), *GetPathNameSafe(Object), SerializedObjectData.Num(), CRC);
 #endif
 				// Cleanup
 				MemoryWriter.Seek(0L);
diff --git a/Engine/Source/Runtime/CoreUObject/Public/Serialization/ArchiveUObject.h b/Engine/Source/Runtime/CoreUObject/Public/Serialization/ArchiveUObject.h
--- a/Engine/Source/Runtime/CoreUObject/Public/Serialization/ArchiveUObject.h
+++ b/Engine/Source/Runtime/CoreUObject/Public/Serialization/ArchiveUObject.h
@@ -1586,7 +1586,10 @@ public:
 	*/
 	uint32 Crc32(UObject* Object, uint32 CRC = 0);
 
-private:
+protected:
+	/** Return if object was already serialized */
+	virtual bool CustomSerialize(class UObject* Object) { return false; }
+
 	/** Internal byte array used for serialization */
 	TArray<uint8> SerializedObjectData;
 	/** Internal archive used for serialization */
diff --git a/Engine/Source/Runtime/Engine/Classes/Engine/Blueprint.h b/Engine/Source/Runtime/Engine/Classes/Engine/Blueprint.h
--- a/Engine/Source/Runtime/Engine/Classes/Engine/Blueprint.h
+++ b/Engine/Source/Runtime/Engine/Classes/Engine/Blueprint.h
@@ -461,7 +461,10 @@ public:
 
 	/** CRC for CDO calculated right after the latest compilation used by Reinstancer to check if default values were changed */
 	UPROPERTY(transient, duplicatetransient)
-	uint32 CrcPreviousCompiledCDO;
+	uint32 CrcLastCompiledCDO;
+
+	UPROPERTY(transient, duplicatetransient)
+	uint32 CrcLastCompiledSignature;
 
 	bool bCachedDependenciesUpToDate;
 	TSet<TWeakObjectPtr<UBlueprint>> CachedDependencies;
