[UE-13974] Not checking ValidateGraphIsWellFormed() for function graphs on skeleton only compiles.

diff --git a/Engine/Source/Editor/KismetCompiler/Private/KismetCompiler.cpp b/Engine/Source/Editor/KismetCompiler/Private/KismetCompiler.cpp
--- a/Engine/Source/Editor/KismetCompiler/Private/KismetCompiler.cpp
+++ b/Engine/Source/Editor/KismetCompiler/Private/KismetCompiler.cpp
@@ -3025,7 +3025,13 @@ void FKismetCompilerContext::ProcessOneFunctionGraph(UEdGraph* SourceGraph, bool
 	// First do some cursory validation (pin types match, inputs to outputs, pins never point to their parent node, etc...)
 	// If this fails we don't proceed any further to avoid crashes or infinite loops
 	// When compiling only the skeleton class, we want the UFunction to be generated and processed so it contains all the local variables, this is unsafe to do during any other compilation mode
-	if (ValidateGraphIsWellFormed(FunctionGraph) || CompileOptions.CompileType == EKismetCompileType::SkeletonOnly)
+	//
+	// NOTE: the order of this conditional check is intentional, and should not
+	//       be rearranged; we do NOT want ValidateGraphIsWellFormed() ran for 
+	//       skeleton-only compiles (that's why we have that check second) 
+	//       because it would most likely result in errors (the function hasn't
+	//       been added to the class yet, etc.)
+	if ((CompileOptions.CompileType == EKismetCompileType::SkeletonOnly) || ValidateGraphIsWellFormed(FunctionGraph))
 	{
 		FKismetFunctionContext& Context = *new (FunctionList)FKismetFunctionContext(MessageLog, Schema, NewClass, Blueprint, CompileOptions.DoesRequireCppCodeGeneration());
 		Context.SourceGraph = FunctionGraph;
