Fix for build failures and crash when running OrionGame OrionEntry

diff --git a/Engine/Source/Editor/KismetCompiler/Private/EdGraphCompiler.cpp b/Engine/Source/Editor/KismetCompiler/Private/EdGraphCompiler.cpp
--- a/Engine/Source/Editor/KismetCompiler/Private/EdGraphCompiler.cpp
+++ b/Engine/Source/Editor/KismetCompiler/Private/EdGraphCompiler.cpp
@@ -141,7 +141,7 @@ void FGraphCompilerContext::FindNodesByClass(const UEdGraph* Graph, TSubclassOf<
 	for (int32 NodeIndex = 0; NodeIndex < Graph->Nodes.Num(); ++NodeIndex)
 	{
 		UEdGraphNode* Node = Graph->Nodes[NodeIndex];
-		if (Node->IsA(NodeClass))
+		if (Node && Node->IsA(NodeClass))
 		{
 			FoundNodes.Add(Node);
 		}
diff --git a/Engine/Source/Editor/KismetCompiler/Private/KismetCompiler.cpp b/Engine/Source/Editor/KismetCompiler/Private/KismetCompiler.cpp
--- a/Engine/Source/Editor/KismetCompiler/Private/KismetCompiler.cpp
+++ b/Engine/Source/Editor/KismetCompiler/Private/KismetCompiler.cpp
@@ -1036,7 +1036,7 @@ void FKismetCompilerContext::PruneIsolatedNodes(const TArray<UEdGraphNode*>& Roo
 	for (int32 NodeIndex = 0; NodeIndex < GraphNodes.Num(); ++NodeIndex)
 	{
 		UEdGraphNode* Node = GraphNodes[NodeIndex];
-		if (!Visitor.VisitedNodes.Contains(Node) && !IsNodePure(Node))
+		if (!Node || !Visitor.VisitedNodes.Contains(Node) && !IsNodePure(Node))
 		{
 			if (!CanIgnoreNode(Node))
 			{
@@ -1044,9 +1044,12 @@ void FKismetCompilerContext::PruneIsolatedNodes(const TArray<UEdGraphNode*>& Roo
 				//MessageLog.Warning(TEXT("Node @@ will never be executed and is being pruned"), Node);
 			}
 
-			if (!ShouldForceKeepNode(Node))
+			if (!Node || !ShouldForceKeepNode(Node))
 			{
-				Node->BreakAllNodeLinks();
+				if (Node)
+				{
+					Node->BreakAllNodeLinks();
+				}
 				GraphNodes.RemoveAtSwap(NodeIndex);
 				--NodeIndex;
 			}
