Fix CIS (|| and && in the same scope).

diff --git a/Engine/Source/Editor/KismetCompiler/Private/KismetCompiler.cpp b/Engine/Source/Editor/KismetCompiler/Private/KismetCompiler.cpp
--- a/Engine/Source/Editor/KismetCompiler/Private/KismetCompiler.cpp
+++ b/Engine/Source/Editor/KismetCompiler/Private/KismetCompiler.cpp
@@ -1036,7 +1036,7 @@ void FKismetCompilerContext::PruneIsolatedNodes(const TArray<UEdGraphNode*>& Roo
 	for (int32 NodeIndex = 0; NodeIndex < GraphNodes.Num(); ++NodeIndex)
 	{
 		UEdGraphNode* Node = GraphNodes[NodeIndex];
-		if (!Node || !Visitor.VisitedNodes.Contains(Node) && !IsNodePure(Node))
+		if (!Node || (!Visitor.VisitedNodes.Contains(Node) && !IsNodePure(Node)))
 		{
 			if (!CanIgnoreNode(Node))
 			{
