Better warnings in KismetCompiler.

diff --git a/Engine/Source/Editor/KismetCompiler/Private/KismetCompiler.cpp b/Engine/Source/Editor/KismetCompiler/Private/KismetCompiler.cpp
--- a/Engine/Source/Editor/KismetCompiler/Private/KismetCompiler.cpp
+++ b/Engine/Source/Editor/KismetCompiler/Private/KismetCompiler.cpp
@@ -1387,7 +1387,7 @@ void FKismetCompilerContext::PrecompileFunction(FKismetFunctionContext& Context)
 			}
 			else
 			{
-				MessageLog.Warning(*LOCTEXT("NoDelegateProperty_Error", "No delegate property found for '%s'").ToString(), *Context.SourceGraph->GetName());
+				MessageLog.Warning(*LOCTEXT("NoDelegateProperty_Error", "No delegate property found for @@").ToString(), Context.SourceGraph);
 			}
 		}
 
@@ -3425,7 +3425,7 @@ void FKismetCompilerContext::Compile()
 			{
 				if(NULL == MCDelegateProp->SignatureFunction)
 				{
-					MessageLog.Warning(TEXT("No SignatureFunction in MulticastDelegateProperty '%s'"), *MCDelegateProp->GetName());
+					MessageLog.Warning(*FString::Printf(TEXT("No SignatureFunction in MulticastDelegateProperty '%s'"), *MCDelegateProp->GetName()));
 				}
 			}
 		}
