UE-12129 This Blueprint will not compile and causes a crash in 4.8

diff --git a/Engine/Source/Editor/KismetCompiler/Private/KismetCompiler.cpp b/Engine/Source/Editor/KismetCompiler/Private/KismetCompiler.cpp
--- a/Engine/Source/Editor/KismetCompiler/Private/KismetCompiler.cpp
+++ b/Engine/Source/Editor/KismetCompiler/Private/KismetCompiler.cpp
@@ -3669,9 +3669,13 @@ void FKismetCompilerContext::Compile()
 			// Get the current dirty state of the package
 			UPackage* const Package = Cast<UPackage>(CurrentBP->GetOutermost());
 			const bool bStartedWithUnsavedChanges = Package != nullptr ? Package->IsDirty() : true;
+			const EBlueprintStatus OriginalStatus = CurrentBP->Status;
 
 			FBlueprintEditorUtils::RefreshExternalBlueprintDependencyNodes(CurrentBP, NewClass);
 			
+			// Dependent blueprints will be recompile anyway by reinstancer (if necessary).
+			CurrentBP->Status = OriginalStatus;
+
 			// Note: We do not send a change notification event to the dependent BP here because
 			// we have not yet reinstanced any of the instances of the BP being compiled, which may
 			// be referenced by instances of the dependent BP that may be reconstructed as a result.
diff --git a/Engine/Source/Editor/UnrealEd/Private/Kismet2/BlueprintEditorUtils.cpp b/Engine/Source/Editor/UnrealEd/Private/Kismet2/BlueprintEditorUtils.cpp
--- a/Engine/Source/Editor/UnrealEd/Private/Kismet2/BlueprintEditorUtils.cpp
+++ b/Engine/Source/Editor/UnrealEd/Private/Kismet2/BlueprintEditorUtils.cpp
@@ -2682,6 +2682,8 @@ void FBlueprintEditorUtils::GatherDependencies(const UBlueprint* InBlueprint, TS
 			}
 		}
 	}
+
+	Dependencies.Remove(InBlueprint);
 }
 
 void FBlueprintEditorUtils::EnsureCachedDependenciesUpToDate(UBlueprint* Blueprint)
