Fix for structurally changed animation blueprint classes crashing during skeleton recompile.

diff --git a/Engine/Source/Editor/KismetCompiler/Private/AnimBlueprintCompiler.cpp b/Engine/Source/Editor/KismetCompiler/Private/AnimBlueprintCompiler.cpp
--- a/Engine/Source/Editor/KismetCompiler/Private/AnimBlueprintCompiler.cpp
+++ b/Engine/Source/Editor/KismetCompiler/Private/AnimBlueprintCompiler.cpp
@@ -977,9 +977,12 @@ void FAnimBlueprintCompiler::CopyTermDefaultsToDefaultObject(UObject* DefaultObj
 	if (bIsDerivedAnimBlueprint)
 	{
 		// If we are a derived animation graph; apply any stored overrides.
-		// Restore values from the root blueprint to catch values where the override has been removed.
-		UAnimBlueprint* RootAnimBP = UAnimBlueprint::FindRootAnimBlueprint(AnimBlueprint);
-		UAnimBlueprintGeneratedClass* RootAnimClass = RootAnimBP->GetAnimBlueprintGeneratedClass();
+		// Restore values from the root class to catch values where the override has been removed.
+		UAnimBlueprintGeneratedClass* RootAnimClass = NewAnimBlueprintClass;
+		while(UAnimBlueprintGeneratedClass* NextClass = Cast<UAnimBlueprintGeneratedClass>(RootAnimClass->GetSuperClass()))
+		{
+			RootAnimClass = NextClass;
+		}
 		UObject* RootDefaultObject = RootAnimClass->GetDefaultObject();
 
 		for (TFieldIterator<UProperty> It(RootAnimClass) ; It; ++It)
