[UE-14836] Automatically generated variables that cache components are no longer readonly, allows components to be deleted again...

diff --git a/Engine/Source/Editor/KismetCompiler/Private/KismetCompiler.cpp b/Engine/Source/Editor/KismetCompiler/Private/KismetCompiler.cpp
--- a/Engine/Source/Editor/KismetCompiler/Private/KismetCompiler.cpp
+++ b/Engine/Source/Editor/KismetCompiler/Private/KismetCompiler.cpp
@@ -623,7 +623,7 @@ void FKismetCompilerContext::CreateClassVariablesFromBlueprint()
 					const FString CategoryName = Node->CategoryName != NAME_None ? Node->CategoryName.ToString() : Blueprint->GetName();
 					
 					NewProperty->SetMetaData(TEXT("Category"), *CategoryName);
-					NewProperty->SetPropertyFlags(CPF_BlueprintVisible | CPF_NonTransactional | CPF_BlueprintReadOnly);
+					NewProperty->SetPropertyFlags(CPF_BlueprintVisible | CPF_NonTransactional );
 				}
 			}
 		}
