[UE-15155] We were not recompiling when blueprints were in the error state, which can happen when blueprints depend on eachother. This would cause reinstancing not to occur when it needed to, and we did a bytecode compile that generated a new (working) blueprint, but there were old instances that had been stubbed out. Typical mismatch between type information and instance data...

diff --git a/Engine/Source/Editor/KismetCompiler/Private/KismetCompilerModule.cpp b/Engine/Source/Editor/KismetCompiler/Private/KismetCompilerModule.cpp
--- a/Engine/Source/Editor/KismetCompiler/Private/KismetCompilerModule.cpp
+++ b/Engine/Source/Editor/KismetCompiler/Private/KismetCompilerModule.cpp
@@ -111,7 +111,7 @@ void FKismet2CompilerModule::CompileBlueprintInner(class UBlueprint* Blueprint,
 			{
 				for (auto CurrentBP : StoredDependentBlueprints)
 				{
-					Reinstancer->EnlistDependentBlueprintToRecompile(CurrentBP, !CurrentBP->IsPossiblyDirty());
+					Reinstancer->EnlistDependentBlueprintToRecompile(CurrentBP, !(CurrentBP->IsPossiblyDirty() || CurrentBP->Status == BS_Error));
 				}
 			}
 		}
