UE-16019 Calling functions that have an implicit WorldContextObject parameter in a Blueprint function library will give a compile error about a missing __WorldContext pin

There is always WorldContext in bp static functions.

diff --git a/Engine/Source/Editor/BlueprintGraph/Classes/K2Node_FunctionEntry.h b/Engine/Source/Editor/BlueprintGraph/Classes/K2Node_FunctionEntry.h
--- a/Engine/Source/Editor/BlueprintGraph/Classes/K2Node_FunctionEntry.h
+++ b/Engine/Source/Editor/BlueprintGraph/Classes/K2Node_FunctionEntry.h
@@ -62,7 +62,5 @@ class UK2Node_FunctionEntry : public UK2Node_FunctionTerminator
 
 	// Returns pin for the automatically added WorldContext parameter (used only by BlueprintFunctionLibrary).
 	BLUEPRINTGRAPH_API UEdGraphPin* GetAutoWorldContextPin() const;
-
-	BLUEPRINTGRAPH_API void RemoveUnnecessaryAutoWorldContext();
 };
 
diff --git a/Engine/Source/Editor/BlueprintGraph/Private/K2Node_FunctionEntry.cpp b/Engine/Source/Editor/BlueprintGraph/Private/K2Node_FunctionEntry.cpp
--- a/Engine/Source/Editor/BlueprintGraph/Private/K2Node_FunctionEntry.cpp
+++ b/Engine/Source/Editor/BlueprintGraph/Private/K2Node_FunctionEntry.cpp
@@ -173,19 +173,7 @@ void UK2Node_FunctionEntry::AllocateDefaultPins()
 
 UEdGraphPin* UK2Node_FunctionEntry::GetAutoWorldContextPin() const
 {
-	return FFunctionEntryHelper::RequireWorldContextParameter(this) ? FindPin(FFunctionEntryHelper::GetWorldContextPinName()) : NULL;
-}
-
-void UK2Node_FunctionEntry::RemoveUnnecessaryAutoWorldContext()
-{
-	auto WorldContextPin = GetAutoWorldContextPin();
-	if (WorldContextPin)
-	{
-		if (!WorldContextPin->LinkedTo.Num())
-		{
-			Pins.Remove(WorldContextPin);
-		}
-	}
+	return FindPin(FFunctionEntryHelper::GetWorldContextPinName());
 }
 
 void UK2Node_FunctionEntry::RemoveOutputPin(UEdGraphPin* PinToRemove)
diff --git a/Engine/Source/Editor/KismetCompiler/Private/KismetCompiler.cpp b/Engine/Source/Editor/KismetCompiler/Private/KismetCompiler.cpp
--- a/Engine/Source/Editor/KismetCompiler/Private/KismetCompiler.cpp
+++ b/Engine/Source/Editor/KismetCompiler/Private/KismetCompiler.cpp
@@ -1166,9 +1166,6 @@ void FKismetCompilerContext::PrecompileFunction(FKismetFunctionContext& Context)
 			TransformNodes(Context);
 		}
 
-		//Now we can safely remove automatically added WorldContext pin from static function.
-		Context.EntryPoint->RemoveUnnecessaryAutoWorldContext();
-
 		// Create the function stub
 		FName NewFunctionName = (Context.EntryPoint->CustomGeneratedFunctionName != NAME_None) ? Context.EntryPoint->CustomGeneratedFunctionName : Context.EntryPoint->SignatureName;
 		if(Context.IsDelegateSignature())
@@ -2536,7 +2533,7 @@ void FKismetCompilerContext::MergeUbergraphPagesIn(UEdGraph* Ubergraph)
 void FKismetCompilerContext::ExpansionStep(UEdGraph* Graph, bool bAllowUbergraphExpansions)
 {
 	// Node expansion may affect the signature of a static function
-	if (bIsFullCompile || Schema->IsStaticFunctionGraph(Graph))
+	if (bIsFullCompile)
 	{
 		BP_SCOPED_COMPILER_EVENT_STAT(EKismetCompilerStats_Expansion);
 
