Skipping skeleton class compile when doing the bytecode compile.  In all existing codepaths the skeleton class is up to date. Saves 3-5% of compile on load test in QAGame

diff --git a/Engine/Source/Editor/KismetCompiler/Private/KismetCompilerModule.cpp b/Engine/Source/Editor/KismetCompiler/Private/KismetCompilerModule.cpp
--- a/Engine/Source/Editor/KismetCompiler/Private/KismetCompilerModule.cpp
+++ b/Engine/Source/Editor/KismetCompiler/Private/KismetCompilerModule.cpp
@@ -143,7 +143,8 @@ void FKismet2CompilerModule::CompileBlueprint(class UBlueprint* Blueprint, const
 		Compiler->PreCompile(Blueprint, CompileOptions);
 	}
 
-	if (CompileOptions.CompileType != EKismetCompileType::Cpp)
+	if (CompileOptions.CompileType != EKismetCompileType::Cpp
+		&& CompileOptions.CompileType != EKismetCompileType::BytecodeOnly )
 	{
 		BP_SCOPED_COMPILER_EVENT_STAT(EKismetCompilerStats_CompileSkeletonClass);
 		auto SkeletonReinstancer = FBlueprintCompileReinstancer::Create(Blueprint->SkeletonGeneratedClass);
