Complex struct (that must be initialized and destroyed) are properly handled in UObject::execLet after AccessViolation exception.
It makes the size of bytecode bigger, but it seems to be necessary.

diff --git a/Engine/Source/Editor/KismetCompiler/Private/KismetCompilerVMBackend.cpp b/Engine/Source/Editor/KismetCompiler/Private/KismetCompilerVMBackend.cpp
--- a/Engine/Source/Editor/KismetCompiler/Private/KismetCompilerVMBackend.cpp
+++ b/Engine/Source/Editor/KismetCompiler/Private/KismetCompilerVMBackend.cpp
@@ -912,6 +912,8 @@ public:
 		else
 		{
 			Writer << EX_Let;
+			Writer << DestinationExpression->AssociatedVarProperty;
+
 		}
 		EmitTerm(DestinationExpression);
 	}
@@ -945,6 +947,8 @@ public:
 		FBPTerminal* TargetExpression = Statement.RHS[1];
 
 		Writer << EX_Let;
+		UProperty* PropertyToHandleComplexStruct = nullptr;
+		Writer << PropertyToHandleComplexStruct;
 		EmitTerm(DestinationExpression);
 
 		Writer << EX_ObjToInterfaceCast;
@@ -961,6 +965,8 @@ public:
 		FBPTerminal* TargetExpression      = Statement.RHS[1];
 
 		Writer << EX_Let;
+		UProperty* PropertyToHandleComplexStruct = nullptr;
+		Writer << PropertyToHandleComplexStruct;
 		EmitTerm(DestinationExpression);
 
 		Writer << EX_CrossInterfaceCast;
@@ -977,6 +983,8 @@ public:
 		FBPTerminal* TargetInterfaceExpression = Statement.RHS[1];
 
 		Writer << EX_Let;
+		UProperty* PropertyToHandleComplexStruct = nullptr;
+		Writer << PropertyToHandleComplexStruct;
 		EmitTerm(DestinationExpression);
 
 		Writer << EX_InterfaceToObjCast;
@@ -993,6 +1001,8 @@ public:
 		FBPTerminal* TargetExpression = Statement.RHS[1];
 
 		Writer << EX_Let;
+		UProperty* PropertyToHandleComplexStruct = nullptr;
+		Writer << PropertyToHandleComplexStruct;
 		EmitTerm(DestinationExpression);
 
 		Writer << EX_DynamicCast;
@@ -1008,6 +1018,8 @@ public:
 		FBPTerminal* TargetExpression = Statement.RHS[1];
 
 		Writer << EX_Let;
+		UProperty* PropertyToHandleComplexStruct = nullptr;
+		Writer << PropertyToHandleComplexStruct;
 		EmitTerm(DestinationExpression);
 
 		Writer << EX_MetaCast;
@@ -1025,6 +1037,8 @@ public:
 		const bool bIsInterfaceCast = (PSCObjClass && PSCObjClass->HasAnyClassFlags(CLASS_Interface));
 
 		Writer << EX_Let;
+		UProperty* PropertyToHandleComplexStruct = nullptr;
+		Writer << PropertyToHandleComplexStruct;
 		EmitTerm(DestinationExpression);
 
 		Writer << EX_PrimitiveCast;
diff --git a/Engine/Source/Editor/UnrealEd/Private/ScriptDisassembler.cpp b/Engine/Source/Editor/UnrealEd/Private/ScriptDisassembler.cpp
--- a/Engine/Source/Editor/UnrealEd/Private/ScriptDisassembler.cpp
+++ b/Engine/Source/Editor/UnrealEd/Private/ScriptDisassembler.cpp
@@ -238,6 +238,8 @@ void FKismetBytecodeDisassembler::ProcessCommon(int32& ScriptIndex, EExprToken O
 			Ar.Logf(TEXT("%s $%X: Let (Variable = Expression)"), *Indents, (int32)Opcode);
 			AddIndent();
 
+			ReadPointer<UProperty>(ScriptIndex);
+
 			// Variable expr.
 			Ar.Logf(TEXT("%s Variable:"), *Indents);
 			SerializeExpr( ScriptIndex );
diff --git a/Engine/Source/Runtime/CoreUObject/Private/UObject/ScriptCore.cpp b/Engine/Source/Runtime/CoreUObject/Private/UObject/ScriptCore.cpp
--- a/Engine/Source/Runtime/CoreUObject/Private/UObject/ScriptCore.cpp
+++ b/Engine/Source/Runtime/CoreUObject/Private/UObject/ScriptCore.cpp
@@ -1343,27 +1343,46 @@ void UObject::execLetValueOnPersistentFrame(FFrame& Stack, RESULT_DECL)
 }
 IMPLEMENT_VM_FUNCTION(Ex_LetValueOnPersistentFrame, execLetValueOnPersistentFrame);
 
-void UObject::execLet( FFrame& Stack, RESULT_DECL )
+void UObject::execLet(FFrame& Stack, RESULT_DECL)
 {
 	checkSlow(!dynamic_cast<UBoolProperty*>(this));
 
+	Stack.MostRecentProperty = nullptr;
+	UProperty* LocallyKnownProperty = Stack.ReadProperty();
+
 	// Get variable address.
-	Stack.MostRecentPropertyAddress = NULL;
-	Stack.Step( Stack.Object, NULL ); // Evaluate variable.
+	Stack.MostRecentProperty = nullptr;
+	Stack.MostRecentPropertyAddress = nullptr;
+	Stack.Step(Stack.Object, nullptr); // Evaluate variable.
 
-	if (Stack.MostRecentPropertyAddress == NULL)
+	uint8* LocalTempResult = nullptr;
+	if (Stack.MostRecentPropertyAddress == nullptr)
 	{
 		static FBlueprintExceptionInfo ExceptionInfo(EBlueprintExceptionType::AccessViolation, TEXT("Attempt to assign variable through None"));
 		FBlueprintCoreDelegates::ThrowScriptException(this, Stack, ExceptionInfo);
 
+		if (LocallyKnownProperty)
+		{
+			LocalTempResult = (uint8*)FMemory_Alloca(LocallyKnownProperty->GetSize());
+			LocallyKnownProperty->InitializeValue(LocalTempResult);
+			Stack.MostRecentPropertyAddress = LocalTempResult;
+		}
+		else
+		{
 			//@TODO: ScriptParallel: Contended static usage
 			static uint8 Crud[1024];//@temp
 			Stack.MostRecentPropertyAddress = Crud;
 			FMemory::Memzero(Stack.MostRecentPropertyAddress, sizeof(FString));
 		}
+	}
 
 	// Evaluate expression into variable.
-	Stack.Step( Stack.Object, Stack.MostRecentPropertyAddress );
+	Stack.Step(Stack.Object, Stack.MostRecentPropertyAddress);
+
+	if (LocalTempResult && LocallyKnownProperty)
+	{
+		LocallyKnownProperty->DestroyValue(LocalTempResult);
+	}
 }
 IMPLEMENT_VM_FUNCTION( EX_Let, execLet );
 
diff --git a/Engine/Source/Runtime/CoreUObject/Public/UObject/ScriptSerialization.h b/Engine/Source/Runtime/CoreUObject/Public/UObject/ScriptSerialization.h
--- a/Engine/Source/Runtime/CoreUObject/Public/UObject/ScriptSerialization.h
+++ b/Engine/Source/Runtime/CoreUObject/Public/UObject/ScriptSerialization.h
@@ -128,6 +128,9 @@
 			break;
 		}
 		case EX_Let:
+		{
+			XFER_PROP_POINTER;
+		}
 		case EX_LetObj:
 		case EX_LetWeakObjPtr:
 		case EX_LetBool:
