UMG/BP - UE-16315 - You can now help the blueprint utility functions create the correct blueprint for you.  Previously all classes created a standard blueprint, now it asks the compiler module if any registered IBlueprintCompiler's have special blueprint classes for a particular UObject class when creating new blueprint assets.  This fixes the problem of the new asset buttons creating UBlueprints for UUserWidget subclasses, instead of UWidgetBlueprints.

diff --git a/Engine/Source/Editor/KismetCompiler/Private/KismetCompilerModule.cpp b/Engine/Source/Editor/KismetCompiler/Private/KismetCompilerModule.cpp
--- a/Engine/Source/Editor/KismetCompiler/Private/KismetCompilerModule.cpp
+++ b/Engine/Source/Editor/KismetCompiler/Private/KismetCompilerModule.cpp
@@ -274,4 +274,18 @@ void FKismet2CompilerModule::RemoveBlueprintGeneratedClasses(class UBlueprint* B
 	}
 }
 
+void FKismet2CompilerModule::GetBlueprintTypesForClass(UClass* ParentClass, UClass*& OutBlueprintClass, UClass*& OutBlueprintGeneratedClass) const
+{
+	for ( IBlueprintCompiler* Compiler : Compilers )
+	{
+		if ( Compiler->GetBlueprintTypesForClass(ParentClass, OutBlueprintClass, OutBlueprintGeneratedClass) )
+		{
+			return;
+		}
+	}
+
+	OutBlueprintClass = UBlueprint::StaticClass();
+	OutBlueprintGeneratedClass = UBlueprintGeneratedClass::StaticClass();
+}
+
 #undef LOCTEXT_NAMESPACE
diff --git a/Engine/Source/Editor/KismetCompiler/Public/KismetCompilerModule.h b/Engine/Source/Editor/KismetCompiler/Public/KismetCompilerModule.h
--- a/Engine/Source/Editor/KismetCompiler/Public/KismetCompilerModule.h
+++ b/Engine/Source/Editor/KismetCompiler/Public/KismetCompilerModule.h
@@ -28,6 +28,13 @@ public:
 		PostCompile(Blueprint);
 	}
 
+	virtual bool GetBlueprintTypesForClass(UClass* ParentClass, UClass*& OutBlueprintClass, UClass*& OutBlueprintGeneratedClass) const
+	{
+		OutBlueprintClass = nullptr;
+		OutBlueprintGeneratedClass = nullptr;
+		return false;
+	}
+
 protected:
 	virtual void PreCompile(UBlueprint* Blueprint) { }
 	virtual void PostCompile(UBlueprint* Blueprint) { }
@@ -71,6 +78,12 @@ public:
 	 * Gets a list of all compilers for blueprints.  You can register new compilers through this list.
 	 */
 	virtual TArray<IBlueprintCompiler*>& GetCompilers() = 0;
+
+	/**
+	 * Get the blueprint class and generated blueprint class for a particular class type.  Not every
+	 * blueprint is a normal UBlueprint, like UUserWidget blueprints should be UWidgetBlueprints.
+	 */
+	virtual void GetBlueprintTypesForClass(UClass* ParentClass, UClass*& OutBlueprintClass, UClass*& OutBlueprintGeneratedClass) const = 0;
 };
 
 //////////////////////////////////////////////////////////////////////////
@@ -88,6 +101,7 @@ public:
 	virtual void RecoverCorruptedBlueprint(class UBlueprint* Blueprint) override;
 	virtual void RemoveBlueprintGeneratedClasses(class UBlueprint* Blueprint) override;
 	virtual TArray<IBlueprintCompiler*>& GetCompilers() override { return Compilers; }
+	virtual void GetBlueprintTypesForClass(UClass* ParentClass, UClass*& OutBlueprintClass, UClass*& OutBlueprintGeneratedClass) const override;
 	// End implementation
 private:
 	void CompileBlueprintInner(class UBlueprint* Blueprint, const FKismetCompilerOptions& CompileOptions, FCompilerResultsLog& Results, TSharedPtr<FBlueprintCompileReinstancer> Reinstancer, TArray<UObject*>* ObjLoaded);
diff --git a/Engine/Source/Editor/UMGEditor/Private/UMGEditorModule.cpp b/Engine/Source/Editor/UMGEditor/Private/UMGEditorModule.cpp
--- a/Engine/Source/Editor/UMGEditor/Private/UMGEditorModule.cpp
+++ b/Engine/Source/Editor/UMGEditor/Private/UMGEditorModule.cpp
@@ -141,6 +141,18 @@ public:
 		}
 	}
 
+	bool GetBlueprintTypesForClass(UClass* ParentClass, UClass*& OutBlueprintClass, UClass*& OutBlueprintGeneratedClass) const override
+	{
+		if ( ParentClass == UUserWidget::StaticClass() || ParentClass->IsChildOf(UUserWidget::StaticClass()) )
+		{
+			OutBlueprintClass = UWidgetBlueprint::StaticClass();
+			OutBlueprintGeneratedClass = UWidgetBlueprintGeneratedClass::StaticClass();
+			return true;
+		}
+
+		return false;
+	}
+
 	/** Gets the extensibility managers for outside entities to extend gui page editor's menus and toolbars */
 	virtual TSharedPtr<FExtensibilityManager> GetMenuExtensibilityManager() override { return MenuExtensibilityManager; }
 	virtual TSharedPtr<FExtensibilityManager> GetToolBarExtensibilityManager() override { return ToolBarExtensibilityManager; }
diff --git a/Engine/Source/Editor/UnrealEd/Private/Kismet2/Kismet2.cpp b/Engine/Source/Editor/UnrealEd/Private/Kismet2/Kismet2.cpp
--- a/Engine/Source/Editor/UnrealEd/Private/Kismet2/Kismet2.cpp
+++ b/Engine/Source/Editor/UnrealEd/Private/Kismet2/Kismet2.cpp
@@ -1521,6 +1521,12 @@ UBlueprint* FKismetEditorUtilities::CreateBlueprintFromClass(FText InWindowTitle
 		NewNameSuggestion = TEXT("NewBlueprint");
 	}
 
+	UClass* BlueprintClass = nullptr;
+	UClass* BlueprintGeneratedClass = nullptr;
+
+	IKismetCompilerInterface& KismetCompilerModule = FModuleManager::LoadModuleChecked<IKismetCompilerInterface>("KismetCompiler");
+	KismetCompilerModule.GetBlueprintTypesForClass(InParentClass, BlueprintClass, BlueprintGeneratedClass);
+
 	FString PackageName = FString(TEXT("/Game/Blueprints/")) + NewNameSuggestion;
 	FString Name;
 	FAssetToolsModule& AssetToolsModule = FModuleManager::LoadModuleChecked<FAssetToolsModule>("AssetTools");
@@ -1550,7 +1556,7 @@ UBlueprint* FKismetEditorUtilities::CreateBlueprintFromClass(FText InWindowTitle
 		check(Package);
 
 		// Create and init a new Blueprint
-		UBlueprint* Blueprint = FKismetEditorUtilities::CreateBlueprint(InParentClass, Package, BPName, BPTYPE_Normal, UBlueprint::StaticClass(), UBlueprintGeneratedClass::StaticClass(), FName("LevelEditorActions"));
+		UBlueprint* Blueprint = FKismetEditorUtilities::CreateBlueprint(InParentClass, Package, BPName, BPTYPE_Normal, BlueprintClass, BlueprintGeneratedClass, FName("LevelEditorActions"));
 		if (Blueprint)
 		{
 			// Notify the asset registry
