Back out changelist 2572336 - This is causing animation blueprints to crash. To reproduce the issue just start QAGame and pie or open the player's animation blueprint

diff --git a/Engine/Source/Editor/KismetCompiler/Private/KismetCompilerVMBackend.cpp b/Engine/Source/Editor/KismetCompiler/Private/KismetCompilerVMBackend.cpp
--- a/Engine/Source/Editor/KismetCompiler/Private/KismetCompilerVMBackend.cpp
+++ b/Engine/Source/Editor/KismetCompiler/Private/KismetCompilerVMBackend.cpp
@@ -912,8 +912,6 @@ public:
 		else
 		{
 			Writer << EX_Let;
-			Writer << DestinationExpression->AssociatedVarProperty;
-
 		}
 		EmitTerm(DestinationExpression);
 	}
@@ -947,8 +945,6 @@ public:
 		FBPTerminal* TargetExpression = Statement.RHS[1];
 
 		Writer << EX_Let;
-		UProperty* PropertyToHandleComplexStruct = nullptr;
-		Writer << PropertyToHandleComplexStruct;
 		EmitTerm(DestinationExpression);
 
 		Writer << EX_ObjToInterfaceCast;
@@ -965,8 +961,6 @@ public:
 		FBPTerminal* TargetExpression      = Statement.RHS[1];
 
 		Writer << EX_Let;
-		UProperty* PropertyToHandleComplexStruct = nullptr;
-		Writer << PropertyToHandleComplexStruct;
 		EmitTerm(DestinationExpression);
 
 		Writer << EX_CrossInterfaceCast;
@@ -983,8 +977,6 @@ public:
 		FBPTerminal* TargetInterfaceExpression = Statement.RHS[1];
 
 		Writer << EX_Let;
-		UProperty* PropertyToHandleComplexStruct = nullptr;
-		Writer << PropertyToHandleComplexStruct;
 		EmitTerm(DestinationExpression);
 
 		Writer << EX_InterfaceToObjCast;
@@ -1001,8 +993,6 @@ public:
 		FBPTerminal* TargetExpression = Statement.RHS[1];
 
 		Writer << EX_Let;
-		UProperty* PropertyToHandleComplexStruct = nullptr;
-		Writer << PropertyToHandleComplexStruct;
 		EmitTerm(DestinationExpression);
 
 		Writer << EX_DynamicCast;
@@ -1018,8 +1008,6 @@ public:
 		FBPTerminal* TargetExpression = Statement.RHS[1];
 
 		Writer << EX_Let;
-		UProperty* PropertyToHandleComplexStruct = nullptr;
-		Writer << PropertyToHandleComplexStruct;
 		EmitTerm(DestinationExpression);
 
 		Writer << EX_MetaCast;
@@ -1037,8 +1025,6 @@ public:
 		const bool bIsInterfaceCast = (PSCObjClass && PSCObjClass->HasAnyClassFlags(CLASS_Interface));
 
 		Writer << EX_Let;
-		UProperty* PropertyToHandleComplexStruct = nullptr;
-		Writer << PropertyToHandleComplexStruct;
 		EmitTerm(DestinationExpression);
 
 		Writer << EX_PrimitiveCast;
diff --git a/Engine/Source/Editor/UnrealEd/Private/ScriptDisassembler.cpp b/Engine/Source/Editor/UnrealEd/Private/ScriptDisassembler.cpp
--- a/Engine/Source/Editor/UnrealEd/Private/ScriptDisassembler.cpp
+++ b/Engine/Source/Editor/UnrealEd/Private/ScriptDisassembler.cpp
@@ -238,8 +238,6 @@ void FKismetBytecodeDisassembler::ProcessCommon(int32& ScriptIndex, EExprToken O
 			Ar.Logf(TEXT("%s $%X: Let (Variable = Expression)"), *Indents, (int32)Opcode);
 			AddIndent();
 
-			ReadPointer<UProperty>(ScriptIndex);
-
 			// Variable expr.
 			Ar.Logf(TEXT("%s Variable:"), *Indents);
 			SerializeExpr( ScriptIndex );
diff --git a/Engine/Source/Runtime/CoreUObject/Private/UObject/ScriptCore.cpp b/Engine/Source/Runtime/CoreUObject/Private/UObject/ScriptCore.cpp
--- a/Engine/Source/Runtime/CoreUObject/Private/UObject/ScriptCore.cpp
+++ b/Engine/Source/Runtime/CoreUObject/Private/UObject/ScriptCore.cpp
@@ -1343,46 +1343,27 @@ void UObject::execLetValueOnPersistentFrame(FFrame& Stack, RESULT_DECL)
 }
 IMPLEMENT_VM_FUNCTION(Ex_LetValueOnPersistentFrame, execLetValueOnPersistentFrame);
 
-void UObject::execLet(FFrame& Stack, RESULT_DECL)
+void UObject::execLet( FFrame& Stack, RESULT_DECL )
 {
 	checkSlow(!dynamic_cast<UBoolProperty*>(this));
 
-	Stack.MostRecentProperty = nullptr;
-	UProperty* LocallyKnownProperty = Stack.ReadProperty();
-
 	// Get variable address.
-	Stack.MostRecentProperty = nullptr;
-	Stack.MostRecentPropertyAddress = nullptr;
-	Stack.Step(Stack.Object, nullptr); // Evaluate variable.
+	Stack.MostRecentPropertyAddress = NULL;
+	Stack.Step( Stack.Object, NULL ); // Evaluate variable.
 
-	uint8* LocalTempResult = nullptr;
-	if (Stack.MostRecentPropertyAddress == nullptr)
+	if (Stack.MostRecentPropertyAddress == NULL)
 	{
 		static FBlueprintExceptionInfo ExceptionInfo(EBlueprintExceptionType::AccessViolation, TEXT("Attempt to assign variable through None"));
 		FBlueprintCoreDelegates::ThrowScriptException(this, Stack, ExceptionInfo);
 
-		if (LocallyKnownProperty)
-		{
-			LocalTempResult = (uint8*)FMemory_Alloca(LocallyKnownProperty->GetSize());
-			LocallyKnownProperty->InitializeValue(LocalTempResult);
-			Stack.MostRecentPropertyAddress = LocalTempResult;
-		}
-		else
-		{
 			//@TODO: ScriptParallel: Contended static usage
 			static uint8 Crud[1024];//@temp
 			Stack.MostRecentPropertyAddress = Crud;
 			FMemory::Memzero(Stack.MostRecentPropertyAddress, sizeof(FString));
 		}
-	}
 
 	// Evaluate expression into variable.
-	Stack.Step(Stack.Object, Stack.MostRecentPropertyAddress);
-
-	if (LocalTempResult && LocallyKnownProperty)
-	{
-		LocallyKnownProperty->DestroyValue(LocalTempResult);
-	}
+	Stack.Step( Stack.Object, Stack.MostRecentPropertyAddress );
 }
 IMPLEMENT_VM_FUNCTION( EX_Let, execLet );
 
diff --git a/Engine/Source/Runtime/CoreUObject/Public/UObject/ScriptSerialization.h b/Engine/Source/Runtime/CoreUObject/Public/UObject/ScriptSerialization.h
--- a/Engine/Source/Runtime/CoreUObject/Public/UObject/ScriptSerialization.h
+++ b/Engine/Source/Runtime/CoreUObject/Public/UObject/ScriptSerialization.h
@@ -128,9 +128,6 @@
 			break;
 		}
 		case EX_Let:
-		{
-			XFER_PROP_POINTER;
-		}
 		case EX_LetObj:
 		case EX_LetWeakObjPtr:
 		case EX_LetBool:
