ME optimization: EX_CallMath - WIP

diff --git a/Engine/Source/Editor/KismetCompiler/Private/KismetCompilerVMBackend.cpp b/Engine/Source/Editor/KismetCompiler/Private/KismetCompilerVMBackend.cpp
--- a/Engine/Source/Editor/KismetCompiler/Private/KismetCompilerVMBackend.cpp
+++ b/Engine/Source/Editor/KismetCompiler/Private/KismetCompilerVMBackend.cpp
@@ -695,7 +695,7 @@ public:
 				NumParams += Param->HasAnyPropertyFlags(CPF_ReturnParm) ? 0 : 1;
 			}
 		}
-		
+
 		// The target label will only ever be set on a call function when calling into the Ubergraph, which requires a patchup
 		// or when re-entering from a latent function which requires a different kind of patchup
 		if ((Statement.TargetLabel != NULL) && !bIsUbergraph)
@@ -739,13 +739,23 @@ public:
 			}
 		}
 
-		// Handle the function calling context if needed
-		FContextEmitter CallContextWriter(*this);
+		const bool bFinalFunction = FunctionToCall->HasAnyFunctionFlags(FUNC_Final) || Statement.bIsParentContext;
+		const bool bMathCall = bFinalFunction
+			&& FunctionToCall->HasAllFunctionFlags(FUNC_Static | FUNC_BlueprintPure | FUNC_Final | FUNC_Native)
+			&& !FunctionToCall->HasAnyFunctionFlags(FUNC_BlueprintAuthorityOnly | FUNC_BlueprintCosmetic)
+			&& !FunctionToCall->GetOuterUClass()->IsChildOf(UInterface::StaticClass())
+			&& FunctionToCall->GetOwnerClass()->GetName() == TEXT("KismetMathLibrary");
 
-		// RValue property is used to clear value after Access Violation. See UObject::ProcessContextOpcod
-		// If the property from LHS is used, then the retured property (with CPF_ReturnParm) is cleared. But properties returned by ref are not cleared. 
-		UProperty* RValueProperty = nullptr; //Statement.LHS ? Statement.LHS->AssociatedVarProperty : nullptr;
-		CallContextWriter.TryStartContext(Statement.FunctionContext, /*bUnsafeToSkip=*/ bHasOutputValue, Statement.bIsInterfaceContext, RValueProperty);
+		if (!bMathCall) // math call doesn't need context
+		{
+			// Handle the function calling context if needed
+			FContextEmitter CallContextWriter(*this);
+
+			// RValue property is used to clear value after Access Violation. See UObject::ProcessContextOpcod
+			// If the property from LHS is used, then the retured property (with CPF_ReturnParm) is cleared. But properties returned by ref are not cleared. 
+			UProperty* RValueProperty = nullptr; //Statement.LHS ? Statement.LHS->AssociatedVarProperty : nullptr;
+			CallContextWriter.TryStartContext(Statement.FunctionContext, /*bUnsafeToSkip=*/ bHasOutputValue, Statement.bIsInterfaceContext, RValueProperty);
+		}
 
 		// Emit the call type
 		if (FunctionToCall->HasAnyFunctionFlags(FUNC_Delegate))
@@ -753,10 +763,17 @@ public:
 			// @todo: Default delegate functions are no longer callable (and also now have mangled names.)  FindField will fail.
 			check(false);
 		}
-		else if (FunctionToCall->HasAnyFunctionFlags(FUNC_Final) || Statement.bIsParentContext)
+		else if (bFinalFunction)
 		{
+			if (bMathCall)
+			{
+				Writer << EX_CallMath;
+			}
+			else
+			{
+				Writer << EX_FinalFunction;
+			}
 			// The function to call doesn't have a native index
-			Writer << EX_FinalFunction;
 			Writer << FunctionToCall;
 		}
 		else
diff --git a/Engine/Source/Editor/UnrealEd/Private/ScriptDisassembler.cpp b/Engine/Source/Editor/UnrealEd/Private/ScriptDisassembler.cpp
--- a/Engine/Source/Editor/UnrealEd/Private/ScriptDisassembler.cpp
+++ b/Engine/Source/Editor/UnrealEd/Private/ScriptDisassembler.cpp
@@ -480,6 +480,17 @@ void FKismetBytecodeDisassembler::ProcessCommon(int32& ScriptIndex, EExprToken O
 			SerializeExpr( ScriptIndex ); // Return expression.
 			break;
 		}
+	case EX_CallMath:
+		{
+			UStruct* StackNode = ReadPointer<UStruct>(ScriptIndex);
+			Ar.Logf(TEXT("%s $%X: Call Math (stack node %s::%s)"), *Indents, (int32)Opcode, *GetNameSafe(StackNode ? StackNode->GetOuter() : nullptr), *GetNameSafe(StackNode));
+
+			while (SerializeExpr(ScriptIndex) != EX_EndFunctionParms)
+			{
+				// Params
+			}
+			break;
+		}
 	case EX_FinalFunction:
 		{
 			UStruct* StackNode = ReadPointer<UStruct>(ScriptIndex);
diff --git a/Engine/Source/Runtime/CoreUObject/Private/UObject/ScriptCore.cpp b/Engine/Source/Runtime/CoreUObject/Private/UObject/ScriptCore.cpp
--- a/Engine/Source/Runtime/CoreUObject/Private/UObject/ScriptCore.cpp
+++ b/Engine/Source/Runtime/CoreUObject/Private/UObject/ScriptCore.cpp
@@ -408,6 +408,27 @@ void UObject::SkipFunction(FFrame& Stack, RESULT_DECL, UFunction* Function)
 #pragma warning (disable : 4750) // warning C4750: function with _alloca() inlined into a loop
 #endif
 
+void UObject::execCallMathFunction(FFrame& Stack, RESULT_DECL)
+{
+	UFunction* Function = (UFunction*)Stack.ReadObject();
+	checkSlow(Function);
+	checkSlow(Function->FunctionFlags & FUNC_Native);
+	UObject* NewContext = Function->GetOuterUClass()->GetDefaultObject(false);
+	checkSlow(NewContext);
+	{
+		FScopeCycleCounterUObject ContextScope(Stack.Object);
+		FScopeCycleCounterUObject FunctionScope(Function);
+
+		// CurrentNativeFunction is used so far only by FLuaContext::InvokeScriptFunction
+		// TGuardValue<UFunction*> NativeFuncGuard(Stack.CurrentNativeFunction, Function);
+		
+		Native Func = Function->GetNativeFunc();
+		checkSlow(Func);
+		(NewContext->*Func)(Stack, RESULT_PARAM);
+	}
+}
+IMPLEMENT_VM_FUNCTION(EX_CallMath, execCallMathFunction);
+
 void UObject::CallFunction( FFrame& Stack, RESULT_DECL, UFunction* Function )
 {
 	checkSlow(Function);
diff --git a/Engine/Source/Runtime/CoreUObject/Public/UObject/Script.h b/Engine/Source/Runtime/CoreUObject/Public/UObject/Script.h
--- a/Engine/Source/Runtime/CoreUObject/Public/UObject/Script.h
+++ b/Engine/Source/Runtime/CoreUObject/Public/UObject/Script.h
@@ -239,6 +239,7 @@ enum EExprToken
 	EX_ArrayConst			= 0x65,
 	EX_EndArrayConst		= 0x66,
 	EX_AssetConst			= 0x67,
+	EX_CallMath				= 0x68, // static pure function from on local call space
 	EX_Max					= 0x100,
 };
 
diff --git a/Engine/Source/Runtime/CoreUObject/Public/UObject/ScriptSerialization.h b/Engine/Source/Runtime/CoreUObject/Public/UObject/ScriptSerialization.h
--- a/Engine/Source/Runtime/CoreUObject/Public/UObject/ScriptSerialization.h
+++ b/Engine/Source/Runtime/CoreUObject/Public/UObject/ScriptSerialization.h
@@ -218,6 +218,7 @@
 			SerializeExpr( iCode, Ar ); // Return expression.
 			break;
 		}
+		case EX_CallMath:
 		case EX_FinalFunction:
 		{
 			XFER_FUNC_POINTER;											// Stack node.
diff --git a/Engine/Source/Runtime/CoreUObject/Public/UObject/UObject.h b/Engine/Source/Runtime/CoreUObject/Public/UObject/UObject.h
--- a/Engine/Source/Runtime/CoreUObject/Public/UObject/UObject.h
+++ b/Engine/Source/Runtime/CoreUObject/Public/UObject/UObject.h
@@ -1142,6 +1142,8 @@ public:
 
 	DECLARE_FUNCTION(execLetValueOnPersistentFrame);
 
+	DECLARE_FUNCTION(execCallMathFunction);
+
 	// -- K2 support functions
 	struct Object_eventExecuteUbergraph_Parms
 	{
