FNodeHandlingFunctor::RegisterLiteral function was separated.

diff --git a/Engine/Source/Editor/KismetCompiler/Private/KismetCompilerMisc.cpp b/Engine/Source/Editor/KismetCompiler/Private/KismetCompilerMisc.cpp
--- a/Engine/Source/Editor/KismetCompiler/Private/KismetCompilerMisc.cpp
+++ b/Engine/Source/Editor/KismetCompiler/Private/KismetCompilerMisc.cpp
@@ -1281,6 +1281,38 @@ void FNodeHandlingFunctor::SanitizeName(FString& Name)
 	}
 }
 
+FBPTerminal* FNodeHandlingFunctor::RegisterLiteral(FKismetFunctionContext& Context, UEdGraphPin* Net)
+{
+	FBPTerminal* Term = nullptr;
+	// Make sure the default value is valid
+	FString DefaultAllowedResult = CompilerContext.GetSchema()->IsCurrentPinDefaultValid(Net);
+	if (!DefaultAllowedResult.IsEmpty())
+	{
+		FText ErrorFormat = LOCTEXT("InvalidDefault_Error", "The current value of the '@@' pin is invalid: {0}");
+		const FText InvalidReasonText = FText::FromString(DefaultAllowedResult);
+
+		FText DefaultValue = FText::FromString(Net->GetDefaultAsString());
+		if (!DefaultValue.IsEmpty())
+		{
+			ErrorFormat = LOCTEXT("InvalidDefaultVal_Error", "The current value ({1}) of the '@@' pin is invalid: {0}");
+		}
+
+		FString ErrorString = FText::Format(ErrorFormat, InvalidReasonText, DefaultValue).ToString();
+		CompilerContext.MessageLog.Error(*ErrorString, Net);
+
+		// Skip over these properties if they are array or ref properties, because the backend can't emit valid code for them
+		if (Net->PinType.bIsArray || Net->PinType.bIsReference)
+		{
+			return nullptr;
+		}
+	}
+
+	Term = Context.RegisterLiteral(Net);
+	Context.NetMap.Add(Net, Term);
+
+	return Term;
+}
+
 void FNodeHandlingFunctor::RegisterNets(FKismetFunctionContext& Context, UEdGraphNode* Node)
 {
 	for (int32 PinIndex = 0; PinIndex < Node->Pins.Num(); ++PinIndex)
@@ -1293,36 +1325,9 @@ void FNodeHandlingFunctor::RegisterNets(FKismetFunctionContext& Context, UEdGrap
 
 			if (Context.NetMap.Find(Net) == NULL)
 			{
-				// New net, resolve the term that will be used to construct it
-				FBPTerminal* Term = NULL;
-
 				if ((Net->Direction == EGPD_Input) && (Net->LinkedTo.Num() == 0))
 				{
-					// Make sure the default value is valid
-					FString DefaultAllowedResult = CompilerContext.GetSchema()->IsCurrentPinDefaultValid(Net);
-					if (!DefaultAllowedResult.IsEmpty())
-					{
-						FText ErrorFormat = LOCTEXT("InvalidDefault_Error", "The current value of the '@@' pin is invalid: {0}");
-						const FText InvalidReasonText = FText::FromString(DefaultAllowedResult);
-
-						FText DefaultValue = FText::FromString(Net->GetDefaultAsString());
-						if (!DefaultValue.IsEmpty())
-						{
-							ErrorFormat = LOCTEXT("InvalidDefaultVal_Error", "The current value ({1}) of the '@@' pin is invalid: {0}");
-						}
-
-						FString ErrorString = FText::Format(ErrorFormat, InvalidReasonText, DefaultValue).ToString();
-						CompilerContext.MessageLog.Error(*ErrorString, Net);
-
-						// Skip over these properties if they are array or ref properties, because the backend can't emit valid code for them
-						if( Pin->PinType.bIsArray || Pin->PinType.bIsReference )
-						{
-							continue;
-						}
-					}
-
-					Term = Context.RegisterLiteral(Net);
-					Context.NetMap.Add(Net, Term);
+					RegisterLiteral(Context, Net);
 				}
 				else
 				{
diff --git a/Engine/Source/Editor/KismetCompiler/Public/KismetCompilerMisc.h b/Engine/Source/Editor/KismetCompiler/Public/KismetCompilerMisc.h
--- a/Engine/Source/Editor/KismetCompiler/Public/KismetCompilerMisc.h
+++ b/Engine/Source/Editor/KismetCompiler/Public/KismetCompilerMisc.h
@@ -121,6 +121,9 @@ protected:
 	// Returns true if the net is *not* a literal, or if it's a literal that is valid.
 	// Returns false only for a bogus literal value.
 	bool ValidateAndRegisterNetIfLiteral(FKismetFunctionContext& Context, UEdGraphPin* Net);
+
+	// Helper to register literal term
+	FBPTerminal* RegisterLiteral(FKismetFunctionContext& Context, UEdGraphPin* Net);
 public:
 	FNodeHandlingFunctor(FKismetCompilerContext& InCompilerContext)
 		: CompilerContext(InCompilerContext)
