Animation: Added explicit preloading of the associated skeleton when compiling an animation blueprint on load, fixing 'phantom' warnings that wouldn't reproduce when you recompiled the AnimBP interactively

diff --git a/Engine/Source/Editor/KismetCompiler/Private/AnimBlueprintCompiler.cpp b/Engine/Source/Editor/KismetCompiler/Private/AnimBlueprintCompiler.cpp
--- a/Engine/Source/Editor/KismetCompiler/Private/AnimBlueprintCompiler.cpp
+++ b/Engine/Source/Editor/KismetCompiler/Private/AnimBlueprintCompiler.cpp
@@ -86,6 +86,12 @@ FAnimBlueprintCompiler::FAnimBlueprintCompiler(UAnimBlueprint* SourceSketch, FCo
 	, AnimBlueprint(SourceSketch)
 	, bIsDerivedAnimBlueprint(false)
 {
+	// Make sure the skeleton has finished preloading
+	if (AnimBlueprint->TargetSkeleton != nullptr)
+	{
+		AnimBlueprint->TargetSkeleton->GetLinker()->Preload(AnimBlueprint->TargetSkeleton);
+	}
+
 	// Determine if there is an anim blueprint in the ancestry of this class
 	bIsDerivedAnimBlueprint = UAnimBlueprint::FindRootAnimBlueprint(AnimBlueprint) != NULL;
 }
