Data layout log for functions in KismetCompiler

diff --git a/Engine/Source/Editor/KismetCompiler/Private/KismetCompiler.cpp b/Engine/Source/Editor/KismetCompiler/Private/KismetCompiler.cpp
--- a/Engine/Source/Editor/KismetCompiler/Private/KismetCompiler.cpp
+++ b/Engine/Source/Editor/KismetCompiler/Private/KismetCompiler.cpp
@@ -3653,12 +3653,17 @@ void FKismetCompilerContext::Compile()
 				UE_LOG(LogK2Compiler, Log, TEXT("%5d:\t%-64s\t%s"), Prop->GetOffset_ForGC(), *GetNameSafe(Prop), *Prop->GetCPPType());
 			}
 
-			if (NewClass->UberGraphFunction)
+			for (auto LocFunction : TFieldRange<UFunction>(NewClass, EFieldIteratorFlags::ExcludeSuper))
 			{
-				UE_LOG(LogK2Compiler, Log, TEXT("\n\nLAYOUT FRAME %s:"), *GetNameSafe(NewClass->UberGraphFunction));
-				for (auto Prop : TFieldRange<UProperty>(NewClass->UberGraphFunction))
+				UE_LOG(LogK2Compiler, Log, TEXT("\n\nLAYOUT FUNCTION %s:"), *GetNameSafe(LocFunction));
+				for (auto Prop : TFieldRange<UProperty>(LocFunction))
 				{
-					UE_LOG(LogK2Compiler, Log, TEXT("%5d:\t%-64s\t%s"), Prop->GetOffset_ForGC(), *GetNameSafe(Prop), *Prop->GetCPPType());
+					const bool bOutParam = Prop && (0 != (Prop->PropertyFlags & CPF_OutParm));
+					const bool bInParam = Prop && !bOutParam && (0 != (Prop->PropertyFlags & CPF_Parm));
+					UE_LOG(LogK2Compiler, Log, TEXT("%5d:\t%-64s\t%s %s%s")
+						, Prop->GetOffset_ForGC(), *GetNameSafe(Prop), *Prop->GetCPPType()
+						, bInParam ? TEXT("Input") : TEXT("")
+						, bOutParam ? TEXT("Output") : TEXT(""));
 				}
 			}
 		}
