[UEBP-60] Removed needless reinstancing of skeleton generated class

diff --git a/Engine/Source/Editor/KismetCompiler/Private/KismetCompilerModule.cpp b/Engine/Source/Editor/KismetCompiler/Private/KismetCompilerModule.cpp
--- a/Engine/Source/Editor/KismetCompiler/Private/KismetCompilerModule.cpp
+++ b/Engine/Source/Editor/KismetCompiler/Private/KismetCompilerModule.cpp
@@ -150,8 +150,7 @@ void FKismet2CompilerModule::CompileBlueprint(class UBlueprint* Blueprint, const
 		&& CompileOptions.bRegenerateSkelton )
 	{
 		BP_SCOPED_COMPILER_EVENT_STAT(EKismetCompilerStats_CompileSkeletonClass);
-		auto SkeletonReinstancer = FBlueprintCompileReinstancer::Create(Blueprint->SkeletonGeneratedClass);
-
+		
 		FCompilerResultsLog SkeletonResults;
 		SkeletonResults.bSilentMode = true;
 		FKismetCompilerOptions SkeletonCompileOptions;
