Back out revision 64 from //depot/UE4/Engine/Source/Editor/KismetCompiler/Private/KismetCompilerModule.cpp
This is needed by the SCS editor

diff --git a/Engine/Source/Editor/KismetCompiler/Private/KismetCompilerModule.cpp b/Engine/Source/Editor/KismetCompiler/Private/KismetCompilerModule.cpp
--- a/Engine/Source/Editor/KismetCompiler/Private/KismetCompilerModule.cpp
+++ b/Engine/Source/Editor/KismetCompiler/Private/KismetCompilerModule.cpp
@@ -150,7 +150,8 @@ void FKismet2CompilerModule::CompileBlueprint(class UBlueprint* Blueprint, const
 		&& CompileOptions.bRegenerateSkelton )
 	{
 		BP_SCOPED_COMPILER_EVENT_STAT(EKismetCompilerStats_CompileSkeletonClass);
-		
+		auto SkeletonReinstancer = FBlueprintCompileReinstancer::Create(Blueprint->SkeletonGeneratedClass);
+
 		FCompilerResultsLog SkeletonResults;
 		SkeletonResults.bSilentMode = true;
 		FKismetCompilerOptions SkeletonCompileOptions;
