Animation: Fixed a crash when creating and opening a new animation blueprint

diff --git a/Engine/Source/Editor/KismetCompiler/Private/AnimBlueprintCompiler.cpp b/Engine/Source/Editor/KismetCompiler/Private/AnimBlueprintCompiler.cpp
--- a/Engine/Source/Editor/KismetCompiler/Private/AnimBlueprintCompiler.cpp
+++ b/Engine/Source/Editor/KismetCompiler/Private/AnimBlueprintCompiler.cpp
@@ -89,7 +89,10 @@ FAnimBlueprintCompiler::FAnimBlueprintCompiler(UAnimBlueprint* SourceSketch, FCo
 	// Make sure the skeleton has finished preloading
 	if (AnimBlueprint->TargetSkeleton != nullptr)
 	{
-		AnimBlueprint->TargetSkeleton->GetLinker()->Preload(AnimBlueprint->TargetSkeleton);
+		if (FLinkerLoad* Linker = AnimBlueprint->TargetSkeleton->GetLinker())
+		{
+			Linker->Preload(AnimBlueprint->TargetSkeleton);
+		}
 	}
 
 	// Determine if there is an anim blueprint in the ancestry of this class
