UE-17122 Having an "IsShiftDown" node with a branch in the On Key up event will cause a crash
- FStructUtils::TheSameLayout check size of the whole structure
- BP Error when incompatible pins are connected
- UWidgetBlueprintLibrary::GetInputEventFromKeyEvent function

diff --git a/Engine/Source/Editor/KismetCompiler/Private/KismetCompiler.cpp b/Engine/Source/Editor/KismetCompiler/Private/KismetCompiler.cpp
--- a/Engine/Source/Editor/KismetCompiler/Private/KismetCompiler.cpp
+++ b/Engine/Source/Editor/KismetCompiler/Private/KismetCompiler.cpp
@@ -307,9 +307,11 @@ void FKismetCompilerContext::ValidateLink(const UEdGraphPin* PinA, const UEdGrap
 	// response will be returned if the pins are not compatible; any other response here then means that the connection is valid.
 	const FPinConnectionResponse ConnectResponse = Schema->CanCreateConnection(PinA, PinB);
 
-	if (ConnectResponse.Response == CONNECT_RESPONSE_DISALLOW)
+	const bool bForbiddenConnection = (ConnectResponse.Response == CONNECT_RESPONSE_DISALLOW);
+	const bool bMissingConversion = (ConnectResponse.Response == CONNECT_RESPONSE_MAKE_WITH_CONVERSION_NODE);
+	if (bForbiddenConnection || bMissingConversion)
 	{
-		MessageLog.Warning(*FString::Printf(*LOCTEXT("PinTypeMismatch_Error", "Can't connect pins @@ and @@: %s").ToString(), *ConnectResponse.Message.ToString()), PinA, PinB);
+		MessageLog.Error(*FString::Printf(*LOCTEXT("PinTypeMismatch_Error", "Can't connect pins @@ and @@: %s").ToString(), *ConnectResponse.Message.ToString()), PinA, PinB);
 	}
 
 	if (PinA && PinB && PinA->Direction != PinB->Direction)
diff --git a/Engine/Source/Runtime/CoreUObject/Private/UObject/Class.cpp b/Engine/Source/Runtime/CoreUObject/Private/UObject/Class.cpp
--- a/Engine/Source/Runtime/CoreUObject/Private/UObject/Class.cpp
+++ b/Engine/Source/Runtime/CoreUObject/Private/UObject/Class.cpp
@@ -4324,7 +4324,10 @@ bool FStructUtils::ArePropertiesTheSame(const UProperty* A, const UProperty* B,
 bool FStructUtils::TheSameLayout(const UStruct* StructA, const UStruct* StructB, bool bCheckPropertiesNames)
 {
 	bool bResult = false;
-	if (StructA && StructB)
+	if (StructA 
+		&& StructB 
+		&& (StructA->GetPropertiesSize() == StructB->GetPropertiesSize())
+		&& (StructA->GetMinAlignment() == StructB->GetMinAlignment()))
 	{
 		const UProperty* PropertyA = StructA->PropertyLink;
 		const UProperty* PropertyB = StructB->PropertyLink;
diff --git a/Engine/Source/Runtime/UMG/Private/WidgetBlueprintLibrary.cpp b/Engine/Source/Runtime/UMG/Private/WidgetBlueprintLibrary.cpp
--- a/Engine/Source/Runtime/UMG/Private/WidgetBlueprintLibrary.cpp
+++ b/Engine/Source/Runtime/UMG/Private/WidgetBlueprintLibrary.cpp
@@ -526,4 +526,9 @@ void UWidgetBlueprintLibrary::GetAllWidgetsWithInterface(UObject* WorldContextOb
 		}
 	}
 }
+
+FInputEvent UWidgetBlueprintLibrary::GetInputEventFromKeyEvent(const FKeyEvent& KeyEvent)
+{
+	return KeyEvent;
+}
 #undef LOCTEXT_NAMESPACE
\ No newline at end of file
diff --git a/Engine/Source/Runtime/UMG/Public/Blueprint/WidgetBlueprintLibrary.h b/Engine/Source/Runtime/UMG/Public/Blueprint/WidgetBlueprintLibrary.h
--- a/Engine/Source/Runtime/UMG/Public/Blueprint/WidgetBlueprintLibrary.h
+++ b/Engine/Source/Runtime/UMG/Public/Blueprint/WidgetBlueprintLibrary.h
@@ -261,4 +261,7 @@ public:
 	*/
 	UFUNCTION(BlueprintCallable, BlueprintCosmetic, Category = "Widget", meta = (WorldContext = "WorldContextObject"))
 	static void GetAllWidgetsWithInterface(UObject* WorldContextObject, TSubclassOf<UInterface> Interface, TArray<UUserWidget*>& FoundWidgets, bool TopLevelOnly);
+
+	UFUNCTION(BlueprintPure, Category = "Widget", meta = (CompactNodeTitle = "->", BlueprintAutocast))
+	static FInputEvent GetInputEventFromKeyEvent(const FKeyEvent& KeyEvent);
 };
