[UE-16488] Adding or removing members from structures and saving referencing blueprints prior to compiling them no longer causes "innapropriate outermost" warnings. Bytecode was referencing the properties that we had moved into the TRASHSTRUCT

diff --git a/Engine/Source/Editor/KismetCompiler/Private/UserDefinedStructureCompilerUtils.cpp b/Engine/Source/Editor/KismetCompiler/Private/UserDefinedStructureCompilerUtils.cpp
--- a/Engine/Source/Editor/KismetCompiler/Private/UserDefinedStructureCompilerUtils.cpp
+++ b/Engine/Source/Editor/KismetCompiler/Private/UserDefinedStructureCompilerUtils.cpp
@@ -74,7 +74,7 @@ struct FUserDefinedStructureCompilerInner
 		}
 	}
 
-	static void CleanAndSanitizeStruct(UUserDefinedStruct* StructToClean)
+	static UObject* CleanAndSanitizeStruct(UUserDefinedStruct* StructToClean)
 	{
 		check(StructToClean);
 
@@ -86,7 +86,7 @@ struct FUserDefinedStructureCompilerInner
 		const FString TransientString = FString::Printf(TEXT("TRASHSTRUCT_%s"), *StructToClean->GetName());
 		const FName TransientName = MakeUniqueObjectName(GetTransientPackage(), UUserDefinedStruct::StaticClass(), FName(*TransientString));
 		UUserDefinedStruct* TransientStruct = NewObject<UUserDefinedStruct>(GetTransientPackage(), TransientName, RF_Public | RF_Transient);
-
+		
 		TArray<UObject*> SubObjects;
 		GetObjectsWithOuter(StructToClean, SubObjects, true);
 		SubObjects.Remove(StructToClean->EditorData);
@@ -114,6 +114,8 @@ struct FUserDefinedStructureCompilerInner
 		StructToClean->ScriptObjectReferences.Empty();
 		StructToClean->PropertyLink = NULL;
 		StructToClean->ErrorMessage.Empty();
+
+		return TransientStruct;
 	}
 
 	static void LogError(UUserDefinedStruct* Struct, FCompilerResultsLog& MessageLog, const FString& ErrorMsg)
@@ -246,6 +248,10 @@ struct FUserDefinedStructureCompilerInner
 			}
 		};
 
+		TArray<UObject*> AllBlueprints;
+		bool const bIncludeDerivedClasses = true;
+		GetObjectsOfClass(UBlueprint::StaticClass(), AllBlueprints, bIncludeDerivedClasses);
+
 		TArray<FDependencyMapEntry> DependencyMap;
 		for (auto Iter = ChangedStructs.CreateConstIterator(); Iter; ++Iter)
 		{
@@ -268,9 +274,36 @@ struct FUserDefinedStructureCompilerInner
 			UUserDefinedStruct* Struct = DependencyMap[StructureToCompileIndex].Struct;
 			check(Struct);
 
-			FUserDefinedStructureCompilerInner::CleanAndSanitizeStruct(Struct);
+			UObject* OldStruct = FUserDefinedStructureCompilerInner::CleanAndSanitizeStruct(Struct);
 			FUserDefinedStructureCompilerInner::InnerCompileStruct(Struct, GetDefault<UEdGraphSchema_K2>(), MessageLog);
 
+			// make sure bytecode references are cleared, compile on load will update unloaded blueprints..
+			{
+				for (auto Blueprint : AllBlueprints)
+				{
+					UBlueprint* AsBlueprint = CastChecked<UBlueprint>(Blueprint);
+					if (UClass* BPClass = AsBlueprint->GeneratedClass)
+					{
+						for (TFieldIterator<UFunction> FuncIter(BPClass, EFieldIteratorFlags::ExcludeSuper); FuncIter; ++FuncIter)
+						{
+							UFunction* CurrentFunction = *FuncIter;
+							if (CurrentFunction->Script.Num() > 0)
+							{
+								TArray<UObject*> NewReferences;
+								// if there's a reference to any of the trashed objects we want to nuke the script:
+								FReferenceFinder ReferenceFinder(NewReferences, OldStruct, false, false, true, true);
+								ReferenceFinder.FindReferences(CurrentFunction);
+								if (NewReferences.Num() > 0)
+								{
+									CurrentFunction->Script.Empty();
+									AsBlueprint->Status = BS_Dirty;
+								}
+							}
+						}
+					}
+				}
+			}
+
 			DependencyMap.RemoveAtSwap(StructureToCompileIndex);
 
 			for (auto EntryIter = DependencyMap.CreateIterator(); EntryIter; ++EntryIter)
