Animation: Added graph validation prior to compilation, and prevented isolated anim node pruning from occurring when there was a compile error, preventing a crash for null nodes

diff --git a/Engine/Source/Editor/KismetCompiler/Private/AnimBlueprintCompiler.cpp b/Engine/Source/Editor/KismetCompiler/Private/AnimBlueprintCompiler.cpp
--- a/Engine/Source/Editor/KismetCompiler/Private/AnimBlueprintCompiler.cpp
+++ b/Engine/Source/Editor/KismetCompiler/Private/AnimBlueprintCompiler.cpp
@@ -495,7 +495,10 @@ void FAnimBlueprintCompiler::PruneIsolatedAnimationNodes(const TArray<UAnimGraph
 void FAnimBlueprintCompiler::ProcessAnimationNodesGivenRoot(TArray<UAnimGraphNode_Base*>& AnimNodeList, const TArray<UAnimGraphNode_Base*>& RootSet)
 {
 	// Now prune based on the root set
-	PruneIsolatedAnimationNodes(RootSet, AnimNodeList);
+	if (MessageLog.NumErrors == 0)
+	{
+		PruneIsolatedAnimationNodes(RootSet, AnimNodeList);
+	}
 
 	// Process the remaining nodes
 	for (auto SourceNodeIt = AnimNodeList.CreateIterator(); SourceNodeIt; ++SourceNodeIt)
@@ -507,6 +510,9 @@ void FAnimBlueprintCompiler::ProcessAnimationNodesGivenRoot(TArray<UAnimGraphNod
 
 void FAnimBlueprintCompiler::ProcessAllAnimationNodes()
 {
+	// Validate the graph
+	ValidateGraphIsWellFormed(ConsolidatedEventGraph);
+
 	// Build the raw node list
 	TArray<UAnimGraphNode_Base*> AnimNodeList;
 	ConsolidatedEventGraph->GetNodesOfClass<UAnimGraphNode_Base>(/*out*/ AnimNodeList);
