UE-17809 Clicking the Add Variable button in a Struct that is referenced as two inputs in a function that is called in an Actor Blueprint's Construction Script causes a crash

diff --git a/Engine/Source/Editor/KismetCompiler/Private/UserDefinedStructureCompilerUtils.cpp b/Engine/Source/Editor/KismetCompiler/Private/UserDefinedStructureCompilerUtils.cpp
--- a/Engine/Source/Editor/KismetCompiler/Private/UserDefinedStructureCompilerUtils.cpp
+++ b/Engine/Source/Editor/KismetCompiler/Private/UserDefinedStructureCompilerUtils.cpp
@@ -13,6 +13,20 @@
 
 struct FUserDefinedStructureCompilerInner
 {
+	static void ClearStructReferencesInBP(UBlueprint* FoundBlueprint, TSet<UBlueprint*>& BlueprintsToRecompile)
+	{
+		bool bAlreadyProcessed = false;
+		BlueprintsToRecompile.Add(FoundBlueprint, &bAlreadyProcessed);
+		if (!bAlreadyProcessed)
+		{
+			for (auto Function : TFieldRange<UFunction>(FoundBlueprint->GeneratedClass, EFieldIteratorFlags::ExcludeSuper))
+			{
+				Function->Script.Empty();
+			}
+			FoundBlueprint->Status = BS_Dirty;
+		}
+	}
+
 	static void ReplaceStructWithTempDuplicate(
 		UUserDefinedStruct* StructureToReinstance, 
 		TSet<UBlueprint*>& BlueprintsToRecompile,
@@ -46,8 +60,8 @@ struct FUserDefinedStructureCompilerInner
 					{
 						if (UBlueprint* FoundBlueprint = Cast<UBlueprint>(OwnerClass->ClassGeneratedBy))
 						{
-							BlueprintsToRecompile.Add(FoundBlueprint);
 							StructProperty->Struct = DuplicatedStruct;
+							ClearStructReferencesInBP(FoundBlueprint, BlueprintsToRecompile);
 						}
 					}
 					else if (auto OwnerStruct = Cast<UUserDefinedStruct>(StructProperty->GetOwnerStruct()))
@@ -71,6 +85,18 @@ struct FUserDefinedStructureCompilerInner
 			}
 
 			DuplicatedStruct->RemoveFromRoot();
+
+			for (auto Blueprint : TObjectRange<UBlueprint>(RF_ClassDefaultObject | RF_PendingKill))
+			{
+				if (Blueprint && !BlueprintsToRecompile.Contains(Blueprint))
+				{
+					FBlueprintEditorUtils::EnsureCachedDependenciesUpToDate(Blueprint);
+					if (Blueprint->CachedUDSDependencies.Contains(StructureToReinstance))
+					{
+						ClearStructReferencesInBP(Blueprint, BlueprintsToRecompile);
+					}
+				}
+			}
 		}
 	}
 
@@ -248,10 +274,6 @@ struct FUserDefinedStructureCompilerInner
 			}
 		};
 
-		TArray<UObject*> AllBlueprints;
-		bool const bIncludeDerivedClasses = true;
-		GetObjectsOfClass(UBlueprint::StaticClass(), AllBlueprints, bIncludeDerivedClasses);
-
 		TArray<FDependencyMapEntry> DependencyMap;
 		for (auto Iter = ChangedStructs.CreateConstIterator(); Iter; ++Iter)
 		{
@@ -274,36 +296,9 @@ struct FUserDefinedStructureCompilerInner
 			UUserDefinedStruct* Struct = DependencyMap[StructureToCompileIndex].Struct;
 			check(Struct);
 
-			UObject* OldStruct = FUserDefinedStructureCompilerInner::CleanAndSanitizeStruct(Struct);
+			FUserDefinedStructureCompilerInner::CleanAndSanitizeStruct(Struct);
 			FUserDefinedStructureCompilerInner::InnerCompileStruct(Struct, GetDefault<UEdGraphSchema_K2>(), MessageLog);
 
-			// make sure bytecode references are cleared, compile on load will update unloaded blueprints..
-			{
-				for (auto Blueprint : AllBlueprints)
-				{
-					UBlueprint* AsBlueprint = CastChecked<UBlueprint>(Blueprint);
-					if (UClass* BPClass = AsBlueprint->GeneratedClass)
-					{
-						for (TFieldIterator<UFunction> FuncIter(BPClass, EFieldIteratorFlags::ExcludeSuper); FuncIter; ++FuncIter)
-						{
-							UFunction* CurrentFunction = *FuncIter;
-							if (CurrentFunction->Script.Num() > 0)
-							{
-								TArray<UObject*> NewReferences;
-								// if there's a reference to any of the trashed objects we want to nuke the script:
-								FReferenceFinder ReferenceFinder(NewReferences, OldStruct, false, false, true, true);
-								ReferenceFinder.FindReferences(CurrentFunction);
-								if (NewReferences.Num() > 0)
-								{
-									CurrentFunction->Script.Empty();
-									AsBlueprint->Status = BS_Dirty;
-								}
-							}
-						}
-					}
-				}
-			}
-
 			DependencyMap.RemoveAtSwap(StructureToCompileIndex);
 
 			for (auto EntryIter = DependencyMap.CreateIterator(); EntryIter; ++EntryIter)
@@ -318,16 +313,13 @@ void FUserDefinedStructureCompilerUtils::CompileStruct(class UUserDefinedStruct*
 {
 	if (FStructureEditorUtils::UserDefinedStructEnabled() && Struct)
 	{
-		TSet<UBlueprint*> BlueprintsThatHaveBeenRecompiled;
-		TSet<UBlueprint*> BlueprintsToRecompile;
-
 		TArray<UUserDefinedStruct*> ChangedStructs; 
-
 		if (FUserDefinedStructureCompilerInner::ShouldBeCompiled(Struct) || bForceRecompile)
 		{
 			ChangedStructs.Add(Struct);
 		}
 
+		TSet<UBlueprint*> BlueprintsToRecompile;
 		for (int32 StructIdx = 0; StructIdx < ChangedStructs.Num(); ++StructIdx)
 		{
 			UUserDefinedStruct* ChangedStruct = ChangedStructs[StructIdx];
@@ -343,6 +335,7 @@ void FUserDefinedStructureCompilerUtils::CompileStruct(class UUserDefinedStruct*
 		FUserDefinedStructureCompilerInner::BuildDependencyMapAndCompile(ChangedStructs, MessageLog);
 
 		// UPDATE ALL THINGS DEPENDENT ON COMPILED STRUCTURES
+		TSet<UBlueprint*> BlueprintsThatHaveBeenRecompiled;
 		for (TObjectIterator<UK2Node> It(RF_Transient | RF_PendingKill | RF_ClassDefaultObject, true); It && ChangedStructs.Num(); ++It)
 		{
 			bool bReconstruct = false;
diff --git a/Engine/Source/Editor/UnrealEd/Private/Kismet2/BlueprintEditorUtils.cpp b/Engine/Source/Editor/UnrealEd/Private/Kismet2/BlueprintEditorUtils.cpp
--- a/Engine/Source/Editor/UnrealEd/Private/Kismet2/BlueprintEditorUtils.cpp
+++ b/Engine/Source/Editor/UnrealEd/Private/Kismet2/BlueprintEditorUtils.cpp
@@ -2752,7 +2752,7 @@ UK2Node_Event* FBlueprintEditorUtils::FindOverrideForFunction(const UBlueprint*
 	return NULL;
 }
 
-void FBlueprintEditorUtils::GatherDependencies(const UBlueprint* InBlueprint, TSet<TWeakObjectPtr<UBlueprint>>& Dependencies)
+void FBlueprintEditorUtils::GatherDependencies(const UBlueprint* InBlueprint, TSet<TWeakObjectPtr<UBlueprint>>& Dependencies, TSet<TWeakObjectPtr<UStruct>>& OutUDSDependencies)
 {
 	struct FGatherDependenciesHelper
 	{
@@ -2786,6 +2786,7 @@ void FBlueprintEditorUtils::GatherDependencies(const UBlueprint* InBlueprint, TS
 
 	check(InBlueprint);
 	Dependencies.Empty();
+	OutUDSDependencies.Empty();
 
 	FGatherDependenciesHelper::ProcessHierarchy(InBlueprint->ParentClass, Dependencies);
 
@@ -2813,7 +2814,14 @@ void FBlueprintEditorUtils::GatherDependencies(const UBlueprint* InBlueprint, TS
 				{
 					for (auto Struct : LocalDependentStructures)
 					{
-						FGatherDependenciesHelper::ProcessHierarchy(Struct, Dependencies);
+						if (auto UDS = Cast<UUserDefinedStruct>(Struct))
+						{
+							OutUDSDependencies.Add(UDS);
+						}
+						else
+						{
+							FGatherDependenciesHelper::ProcessHierarchy(Struct, Dependencies);
+						}
 					}
 				}
 			}
@@ -2827,7 +2835,7 @@ void FBlueprintEditorUtils::EnsureCachedDependenciesUpToDate(UBlueprint* Bluepri
 {
 	if (Blueprint && !Blueprint->bCachedDependenciesUpToDate)
 	{
-		GatherDependencies(Blueprint, Blueprint->CachedDependencies);
+		GatherDependencies(Blueprint, Blueprint->CachedDependencies, Blueprint->CachedUDSDependencies);
 		Blueprint->bCachedDependenciesUpToDate = true;
 	}
 }
diff --git a/Engine/Source/Editor/UnrealEd/Public/Kismet2/BlueprintEditorUtils.h b/Engine/Source/Editor/UnrealEd/Public/Kismet2/BlueprintEditorUtils.h
--- a/Engine/Source/Editor/UnrealEd/Public/Kismet2/BlueprintEditorUtils.h
+++ b/Engine/Source/Editor/UnrealEd/Public/Kismet2/BlueprintEditorUtils.h
@@ -395,7 +395,7 @@ public:
 	}
 
 	/** Gather all bps that Blueprint depends on */
-	static void GatherDependencies(const UBlueprint* Blueprint, TSet<TWeakObjectPtr<UBlueprint>>& OutDependencies);
+	static void GatherDependencies(const UBlueprint* Blueprint, TSet<TWeakObjectPtr<UBlueprint>>& OutDependencies, TSet<TWeakObjectPtr<UStruct>>& OutUDSDependencies);
 
 	/** Returns a list of loaded Blueprints that are dependent on the given Blueprint. */
 	static void GetDependentBlueprints(UBlueprint* Blueprint, TArray<UBlueprint*>& DependentBlueprints);
diff --git a/Engine/Source/Runtime/Engine/Classes/Engine/Blueprint.h b/Engine/Source/Runtime/Engine/Classes/Engine/Blueprint.h
--- a/Engine/Source/Runtime/Engine/Classes/Engine/Blueprint.h
+++ b/Engine/Source/Runtime/Engine/Classes/Engine/Blueprint.h
@@ -475,6 +475,9 @@ public:
 	bool bCachedDependenciesUpToDate;
 	TSet<TWeakObjectPtr<UBlueprint>> CachedDependencies;
 
+	// User Defined Structures, the blueprint depends on
+	TSet<TWeakObjectPtr<UStruct>> CachedUDSDependencies;
+
 	bool IsUpToDate() const
 	{
 		return BS_UpToDate == Status || BS_UpToDateWithWarnings == Status;
