Putting GBlueprintCompileTime back in - it's useful even if it misses a few compile related functions

diff --git a/Engine/Source/Editor/KismetCompiler/Private/KismetCompilerModule.cpp b/Engine/Source/Editor/KismetCompiler/Private/KismetCompilerModule.cpp
--- a/Engine/Source/Editor/KismetCompiler/Private/KismetCompilerModule.cpp
+++ b/Engine/Source/Editor/KismetCompiler/Private/KismetCompilerModule.cpp
@@ -133,6 +133,7 @@ extern UNREALED_API FSecondsCounterData BlueprintCompileAndLoadTimerData;
 // Compiles a blueprint.
 void FKismet2CompilerModule::CompileBlueprint(class UBlueprint* Blueprint, const FKismetCompilerOptions& CompileOptions, FCompilerResultsLog& Results, TSharedPtr<FBlueprintCompileReinstancer> ParentReinstancer, TArray<UObject*>* ObjLoaded)
 {
+	SCOPE_SECONDS_COUNTER(GBlueprintCompileTime);
 	FSecondsCounterScope Timer(BlueprintCompileAndLoadTimerData);
 	BP_SCOPED_COMPILER_EVENT_STAT(EKismetCompilerStats_CompileTime);
 
