[NO NON-PII COMMIT MESSAGE]

diff --git a/Engine/Source/Editor/KismetCompiler/Private/KismetCompilerMisc.cpp b/Engine/Source/Editor/KismetCompiler/Private/KismetCompilerMisc.cpp
--- a/Engine/Source/Editor/KismetCompiler/Private/KismetCompilerMisc.cpp
+++ b/Engine/Source/Editor/KismetCompiler/Private/KismetCompilerMisc.cpp
@@ -601,13 +601,13 @@ UEdGraphPin* FKismetCompilerUtilities::GenerateAssignmentNodes(class FKismetComp
 					continue;
 				}
 
-				if( ForClass->ClassDefaultObject )
+				if (ForClass->ClassDefaultObject)
 				{
 					// We don't want to generate an assignment node unless the default value 
 					// differs from the value in the CDO:
 					FString DefaultValueAsString;
 					FBlueprintEditorUtils::PropertyValueToString(Property, (uint8*)ForClass->ClassDefaultObject, DefaultValueAsString);
-					if (DefaultValueAsString == OrgPin->DefaultValue)
+					if (DefaultValueAsString == OrgPin->GetDefaultAsString())
 					{
 						continue;
 					}
