MovieScene: Minor cleanups

diff --git a/Engine/Source/Editor/KismetCompiler/Private/KismetCompiler.cpp b/Engine/Source/Editor/KismetCompiler/Private/KismetCompiler.cpp
--- a/Engine/Source/Editor/KismetCompiler/Private/KismetCompiler.cpp
+++ b/Engine/Source/Editor/KismetCompiler/Private/KismetCompiler.cpp
@@ -2076,14 +2076,14 @@ void FKismetCompilerContext::ExpandTimelineNodes(UEdGraph* SourceGraph)
 
 UEdGraphPin* FKismetCompilerContext::ExpandNodesToAllocateRuntimeMovieScenePlayer( UEdGraph* SourceGraph, UK2Node_PlayMovieScene* PlayMovieSceneNode, ULevel* Level, UK2Node_TemporaryVariable*& OutPlayerVariableNode )
 {
-	// Call URuntimeMovieScenePlayer::CreateRuntimeMovieScenePlayer() to create a new RuntimeMovieScenePlayer instance
+	// Call URuntimeMovieScenePlayer::CreatePlayer() to create a new RuntimeMovieScenePlayer instance
 	UK2Node_CallFunction* CreatePlayerCallNode = SpawnIntermediateNode<UK2Node_CallFunction>( PlayMovieSceneNode, SourceGraph );
 	{
-		CreatePlayerCallNode->FunctionReference.SetExternalMember(GET_FUNCTION_NAME_CHECKED(URuntimeMovieScenePlayer, CreateRuntimeMovieScenePlayer), URuntimeMovieScenePlayer::StaticClass());
+		CreatePlayerCallNode->FunctionReference.SetExternalMember(GET_FUNCTION_NAME_CHECKED(URuntimeMovieScenePlayer, CreatePlayer), URuntimeMovieScenePlayer::StaticClass());
 		CreatePlayerCallNode->AllocateDefaultPins();
 	}
 
-	// The return value of URuntimeMovieScenePlayer::CreateRuntimeMovieScenePlayer() is the actual MovieScenePlayer we'll be operating with
+	// The return value of URuntimeMovieScenePlayer::CreatePlayer() is the actual MovieScenePlayer we'll be operating with
 	UEdGraphPin* CreatePlayerReturnValuePin = CreatePlayerCallNode->GetReturnValuePin();
 
 	// Make a literal for the level and bind it to our function call as a parameter
@@ -2143,7 +2143,7 @@ UEdGraphPin* FKismetCompilerContext::ExpandNodesToAllocateRuntimeMovieScenePlaye
 		// create the player object
 		IfVariableNullNode->GetThenPin()->MakeLinkTo( CreatePlayerCallNode->GetExecPin() );
 
-		// Setup function params for "URuntimeMovieScenePlayer::CreateRuntimeMovieScenePlayer()"
+		// Setup function params for "URuntimeMovieScenePlayer::CreatePlayer()"
 		{
 			// Our level literal just points to the level object
 			LevelLiteralNode->SetObjectRef( Level );
diff --git a/Engine/Source/Runtime/MovieScene/Private/RuntimeMovieScenePlayer.cpp b/Engine/Source/Runtime/MovieScene/Private/RuntimeMovieScenePlayer.cpp
--- a/Engine/Source/Runtime/MovieScene/Private/RuntimeMovieScenePlayer.cpp
+++ b/Engine/Source/Runtime/MovieScene/Private/RuntimeMovieScenePlayer.cpp
@@ -9,25 +9,25 @@
 /* URuntimeMovieScenePlayer static functions
  *****************************************************************************/
 
-URuntimeMovieScenePlayer* URuntimeMovieScenePlayer::CreateRuntimeMovieScenePlayer(ULevel* Level, UMovieSceneBindings* MovieSceneBindings)
-{	
-	URuntimeMovieScenePlayer* NewRuntimeMovieScenePlayer = NewObject<URuntimeMovieScenePlayer>((UObject*)GetTransientPackage(), NAME_None, RF_Transient);
-	check(NewRuntimeMovieScenePlayer != nullptr);
+URuntimeMovieScenePlayer* URuntimeMovieScenePlayer::CreatePlayer(ULevel* Level, UMovieSceneBindings* MovieSceneBindings)
+{
+	URuntimeMovieScenePlayer* NewPlayer = NewObject<URuntimeMovieScenePlayer>((UObject*)GetTransientPackage(), NAME_None, RF_Transient);
+	check(NewPlayer != nullptr);
 
-	// Associate the player with its world
-	NewRuntimeMovieScenePlayer->World = Level->OwningWorld;
-	check( NewRuntimeMovieScenePlayer->World.IsValid() );
+	// associate the player with its world
+	NewPlayer->World = Level->OwningWorld;
+	check(NewPlayer->World.IsValid());
 
-	// Attach bindings to the newly-created RuntimeMovieScenePlayer
-	NewRuntimeMovieScenePlayer->SetMovieSceneBindings( MovieSceneBindings );
+	// attach bindings to the newly-created RuntimeMovieScenePlayer
+	NewPlayer->SetMovieSceneBindings(MovieSceneBindings);
 
-	// Spawn actors for each spawnable!
+	// spawn actors for each spawnable!
 
-	// Add this to the level's list of active RuntimeMovieScenePlayers.  The level will own this player from now on.
-	check( Level != nullptr );
-	Level->AddActiveRuntimeMovieScenePlayer( NewRuntimeMovieScenePlayer );
+	// add to the level's list of active players
+	check(Level != nullptr);
+	Level->AddActiveRuntimeMovieScenePlayer(NewPlayer);
 
-	return NewRuntimeMovieScenePlayer;
+	return NewPlayer;
 }
 
 
diff --git a/Engine/Source/Runtime/MovieScene/Public/RuntimeMovieScenePlayer.h b/Engine/Source/Runtime/MovieScene/Public/RuntimeMovieScenePlayer.h
--- a/Engine/Source/Runtime/MovieScene/Public/RuntimeMovieScenePlayer.h
+++ b/Engine/Source/Runtime/MovieScene/Public/RuntimeMovieScenePlayer.h
@@ -38,7 +38,7 @@ public:
 	 * @return The newly created RuntimeMovieScenePlayer instance for this MovieScene
 	 */
 	UFUNCTION(BlueprintCallable, Category="Game|Cinematic")
-	static URuntimeMovieScenePlayer* CreateRuntimeMovieScenePlayer(ULevel* Level, UMovieSceneBindings* InMovieSceneBindings);
+	static URuntimeMovieScenePlayer* CreatePlayer(ULevel* Level, UMovieSceneBindings* InMovieSceneBindings);
 
 public:
 
