[NO NON-PII COMMIT MESSAGE]

diff --git a/Engine/Source/Editor/BlueprintGraph/Private/EdGraphSchema_K2.cpp b/Engine/Source/Editor/BlueprintGraph/Private/EdGraphSchema_K2.cpp
--- a/Engine/Source/Editor/BlueprintGraph/Private/EdGraphSchema_K2.cpp
+++ b/Engine/Source/Editor/BlueprintGraph/Private/EdGraphSchema_K2.cpp
@@ -4087,12 +4087,17 @@ void UEdGraphSchema_K2::ValidateExistingConnections(UEdGraphPin* Pin)
 {
 	const UEdGraphSchema_K2* K2Schema = GetDefault<UEdGraphSchema_K2>();
 
+	const UBlueprint* Blueprint = FBlueprintEditorUtils::FindBlueprintForNode(Pin->GetOwningNodeUnchecked());
+	const UClass* CallingContext = Blueprint
+		? ((Blueprint->GeneratedClass != nullptr) ? Blueprint->GeneratedClass : Blueprint->ParentClass)
+		: nullptr;
+
 	// Break any newly invalid links
 	TArray<UEdGraphPin*> BrokenLinks;
-	for (int32 Index = 0; Index < Pin->LinkedTo.Num(); )
+	for (int32 Index = 0; Index < Pin->LinkedTo.Num();)
 	{
 		UEdGraphPin* OtherPin = Pin->LinkedTo[Index];
-		if (K2Schema->ArePinsCompatible(Pin, OtherPin))
+		if (K2Schema->ArePinsCompatible(Pin, OtherPin, CallingContext))
 		{
 			++Index;
 		}
diff --git a/Engine/Source/Editor/KismetCompiler/Private/KismetCompilerVMBackend.cpp b/Engine/Source/Editor/KismetCompiler/Private/KismetCompilerVMBackend.cpp
--- a/Engine/Source/Editor/KismetCompiler/Private/KismetCompilerVMBackend.cpp
+++ b/Engine/Source/Editor/KismetCompiler/Private/KismetCompilerVMBackend.cpp
@@ -1259,6 +1259,9 @@ public:
 
 	void EmitSwitchValue(FBlueprintCompiledStatement& Statement)
 	{
+		const int32 TermsBeforeCases = 1;
+		const int32 TermsPerCase = 2;
+
 		if ((Statement.RHS.Num() < 4) || (1 == (Statement.RHS.Num() % 2)))
 		{
 			// Error
@@ -1267,29 +1270,34 @@ public:
 
 		Writer << EX_SwitchValue;
 		// number of cases (without default)
-		const int32 TermsPerCase = 2;
 		uint16 NumCases = ((Statement.RHS.Num() - 2) / TermsPerCase);
 		Writer << NumCases;
 		// end goto index
 		CodeSkipSizeType PatchUpNeededAtOffset = Writer.EmitPlaceholderSkip();
-		// index term
-		const int32 TermsBeforeCases = 1;
-		EmitTerm(Statement.RHS[0]);
 
-		UProperty* VirtualIndexProperty = Statement.RHS[0]->AssociatedVarProperty;
+		// index term
+		auto IndexTerm = Statement.RHS[0];
+		check(IndexTerm);
+		EmitTerm(IndexTerm);
+		UProperty* VirtualIndexProperty = IndexTerm->AssociatedVarProperty;
 		check(VirtualIndexProperty);
 
+		auto DefaultTerm = Statement.RHS[TermsBeforeCases + NumCases*TermsPerCase];
+		check(DefaultTerm);
+		UProperty* VirtualValueProperty = DefaultTerm->AssociatedVarProperty;
+		check(VirtualValueProperty);
+
 		for (uint16 TermIndex = TermsBeforeCases; TermIndex < (NumCases * TermsPerCase); ++TermIndex)
 		{
 			EmitTerm(Statement.RHS[TermIndex], VirtualIndexProperty); // it's a literal value
 			++TermIndex;
 			CodeSkipSizeType PatchOffsetToNextCase = Writer.EmitPlaceholderSkip();
-			EmitTerm(Statement.RHS[TermIndex]);  // it's not a literal
+			EmitTerm(Statement.RHS[TermIndex], VirtualValueProperty);  // it could be literal for 'self'
 			Writer.CommitSkip(PatchOffsetToNextCase, Writer.ScriptBuffer.Num());
 		}
 
 		// default term
-		EmitTerm(Statement.RHS[TermsBeforeCases + NumCases*TermsPerCase]);
+		EmitTerm(DefaultTerm);
 
 		Writer.CommitSkip(PatchUpNeededAtOffset, Writer.ScriptBuffer.Num());
 	}
