[UE-18758] Downgraded back to a warning. This cannot be an error until false positive type errors are fixed (see UE-14123)

diff --git a/Engine/Source/Editor/KismetCompiler/Private/KismetCompiler.cpp b/Engine/Source/Editor/KismetCompiler/Private/KismetCompiler.cpp
--- a/Engine/Source/Editor/KismetCompiler/Private/KismetCompiler.cpp
+++ b/Engine/Source/Editor/KismetCompiler/Private/KismetCompiler.cpp
@@ -311,7 +311,7 @@ void FKismetCompilerContext::ValidateLink(const UEdGraphPin* PinA, const UEdGrap
 	const bool bMissingConversion = (ConnectResponse.Response == CONNECT_RESPONSE_MAKE_WITH_CONVERSION_NODE);
 	if (bForbiddenConnection || bMissingConversion)
 	{
-		MessageLog.Error(*FString::Printf(*LOCTEXT("PinTypeMismatch_Error", "Can't connect pins @@ and @@: %s").ToString(), *ConnectResponse.Message.ToString()), PinA, PinB);
+		MessageLog.Warning(*FString::Printf(*LOCTEXT("PinTypeMismatch_Error", "Can't connect pins @@ and @@: %s").ToString(), *ConnectResponse.Message.ToString()), PinA, PinB);
 	}
 
 	if (PinA && PinB && PinA->Direction != PinB->Direction)
