UE-18925 Cannot Rename or Delete Event Dispatchers if they are named "color"

diff --git a/Engine/Source/Editor/KismetCompiler/Private/KismetCompiler.cpp b/Engine/Source/Editor/KismetCompiler/Private/KismetCompiler.cpp
--- a/Engine/Source/Editor/KismetCompiler/Private/KismetCompiler.cpp
+++ b/Engine/Source/Editor/KismetCompiler/Private/KismetCompiler.cpp
@@ -456,17 +456,12 @@ void FKismetCompilerContext::ValidateVariableNames()
 			}
 			else if (ParentClass->HasAnyFlags(RF_Native)) // the above case handles when the parent is a blueprint
 			{
-				UClass* SuperClass = ParentClass;
-				do
+				if (auto ExisingField = FindField<UField>(ParentClass, *VarNameStr))
 				{
-					if (FindObject<UObject>(SuperClass, *VarNameStr, /*ExactClass =*/false))
-					{
-						NewVarName = FBlueprintEditorUtils::FindUniqueKismetName(Blueprint, VarNameStr);
-						break;
-					}
-					SuperClass = SuperClass->GetSuperClass();
-
-				} while (SuperClass != nullptr);
+					UE_LOG(LogK2Compiler, Warning, TEXT("ValidateVariableNames name %s (used in %s) is already taken by %s")
+						, *VarNameStr, *Blueprint->GetPathName(), *ExisingField->GetPathName());
+					NewVarName = FBlueprintEditorUtils::FindUniqueKismetName(Blueprint, VarNameStr);
+				}
 			}
 
 			if (OldVarName != NewVarName)
