MovieScene: Refactored MovieScene and Sequencer to unify the implementation and handling of movie scenes and object bindings

- IMovieSceneAnimation derivatives are now the primary types consumed by Sequencer; they hold a root MovieScene inside
- Object binding managers have been removed; bindings are persisted in UActorAnimation, UNiagaraAnimation, UWidgetAnimation
- MovieSceneAssetEditor is now ActorAnimationEditor and became a plug-in
- MovieSceneRuntimePlayer is now ActorAnimationPlayer; K2Node_PlayMovieScene has been removed

diff --git a/Engine/Plugins/MovieScene/ActorAnimationEditor/ActorAnimationEditor.uplugin b/Engine/Plugins/MovieScene/ActorAnimationEditor/ActorAnimationEditor.uplugin
new file mode 100644
--- /dev/null
+++ b/Engine/Plugins/MovieScene/ActorAnimationEditor/ActorAnimationEditor.uplugin
@@ -0,0 +1,23 @@
+{
+	"FileVersion" : 3,
+	
+	"FriendlyName" : "Actor Animation Editor",
+	"Version" : 1,
+	"VersionName" : "1.0",
+	"CreatedBy" : "Epic Games, Inc.",
+	"CreatedByURL" : "http://epicgames.com",
+	"EngineVersion" : "4.5.0",
+	"Description" : "Content Editor for ActorAnimation Assets.",
+	"Category" : "Editor",
+	"EnabledByDefault" : true,
+
+	"Modules" :
+	[
+		{
+			"Name" : "ActorAnimationEditor",
+			"Type" : "Editor"
+		}
+	],
+
+	"CanContainContent" : false
+}
\ No newline at end of file
diff --git a/Engine/Plugins/MovieScene/ActorAnimationEditor/Source/ActorAnimationEditor/ActorAnimationEditor.Build.cs b/Engine/Plugins/MovieScene/ActorAnimationEditor/Source/ActorAnimationEditor/ActorAnimationEditor.Build.cs
new file mode 100644
--- /dev/null
+++ b/Engine/Plugins/MovieScene/ActorAnimationEditor/Source/ActorAnimationEditor/ActorAnimationEditor.Build.cs
@@ -0,0 +1,37 @@
+// Copyright 1998-2015 Epic Games, Inc. All Rights Reserved.
+
+using UnrealBuildTool;
+
+public class ActorAnimationEditor : ModuleRules
+{
+	public ActorAnimationEditor(TargetInfo Target)
+	{
+        DynamicallyLoadedModuleNames.AddRange(
+            new string[] {
+				"AssetTools",
+			}
+        );
+        
+        PrivateDependencyModuleNames.AddRange(
+			new string[] {
+                "ActorAnimation",
+				"BlueprintGraph",
+				"Core",
+				"CoreUObject",
+                "EditorStyle",
+                "Engine",
+				"MovieScene",
+				"Sequencer",
+                "Slate",
+                "SlateCore",
+                "UnrealEd",
+			}
+		);
+
+		PrivateIncludePathModuleNames.AddRange(
+			new string[] {
+				"AssetTools",
+			}
+		);
+	}
+}
diff --git a/Engine/Plugins/MovieScene/ActorAnimationEditor/Source/ActorAnimationEditor/Private/ActorAnimationActions.cpp b/Engine/Plugins/MovieScene/ActorAnimationEditor/Source/ActorAnimationEditor/Private/ActorAnimationActions.cpp
new file mode 100644
--- /dev/null
+++ b/Engine/Plugins/MovieScene/ActorAnimationEditor/Source/ActorAnimationEditor/Private/ActorAnimationActions.cpp
@@ -0,0 +1,72 @@
+// Copyright 1998-2015 Epic Games, Inc. All Rights Reserved.
+
+#include "ActorAnimationEditorPrivatePCH.h"
+#include "IAssetTools.h"
+#include "ISequencerModule.h"
+#include "Toolkits/IToolkit.h"
+
+
+#define LOCTEXT_NAMESPACE "AssetTypeActions"
+
+
+/* IAssetTypeActions interface
+ *****************************************************************************/
+
+uint32 FActorAnimationActions::GetCategories()
+{
+	return EAssetTypeCategories::Animation;
+}
+
+
+FText FActorAnimationActions::GetName() const
+{
+	return NSLOCTEXT("AssetTypeActions", "AssetTypeActions_ActorAnimation", "Actor Animation");
+}
+
+
+UClass* FActorAnimationActions::GetSupportedClass() const
+{
+	return UActorAnimation::StaticClass(); 
+}
+
+
+FColor FActorAnimationActions::GetTypeColor() const
+{
+	return FColor(200, 80, 80);
+}
+
+
+void FActorAnimationActions::OpenAssetEditor(const TArray<UObject*>& InObjects, TSharedPtr<IToolkitHost> EditWithinLevelEditor)
+{
+	// @todo sequencer: Only allow users to create new MovieScenes if that feature is turned on globally.
+	if (!FParse::Param(FCommandLine::Get(), TEXT("Sequencer")))
+	{
+		return;
+	}
+
+	EToolkitMode::Type Mode = EditWithinLevelEditor.IsValid()
+		? EToolkitMode::WorldCentric
+		: EToolkitMode::Standalone;
+
+	for (auto ObjIt = InObjects.CreateConstIterator(); ObjIt; ++ObjIt)
+	{
+		UActorAnimation* ActorAnimation = Cast<UActorAnimation>(*ObjIt);
+
+		if (ActorAnimation != nullptr)
+		{
+			FSequencerViewParams ViewParams(TEXT("MovieSceneEditorSettings"));
+			TSharedRef<FActorAnimationEditorToolkit> Toolkit = MakeShareable(new FActorAnimationEditorToolkit);
+			Toolkit->Initialize(Mode, ViewParams, EditWithinLevelEditor, ActorAnimation, true);
+		}
+	}
+}
+
+
+bool FActorAnimationActions::ShouldForceWorldCentric()
+{
+	// @todo sequencer: Hack to force world-centric mode for Sequencer
+	return true;
+}
+
+
+#undef LOCTEXT_NAMESPACE
\ No newline at end of file
diff --git a/Engine/Plugins/MovieScene/ActorAnimationEditor/Source/ActorAnimationEditor/Private/ActorAnimationActions.h b/Engine/Plugins/MovieScene/ActorAnimationEditor/Source/ActorAnimationEditor/Private/ActorAnimationActions.h
new file mode 100644
--- /dev/null
+++ b/Engine/Plugins/MovieScene/ActorAnimationEditor/Source/ActorAnimationEditor/Private/ActorAnimationActions.h
@@ -0,0 +1,24 @@
+// Copyright 1998-2015 Epic Games, Inc. All Rights Reserved.
+
+#pragma once
+
+#include "AssetTypeActions_Base.h"
+
+
+/**
+ * Implements actions for UActorAnimation assets.
+ */
+class FActorAnimationActions
+	: public FAssetTypeActions_Base
+{
+public:
+	
+	// IAssetTypeActions interface
+
+	virtual uint32 GetCategories() override;
+	virtual FText GetName() const override;
+	virtual UClass* GetSupportedClass() const override;
+	virtual FColor GetTypeColor() const override;
+	virtual void OpenAssetEditor(const TArray<UObject*>& InObjects, TSharedPtr<class IToolkitHost> EditWithinLevelEditor = TSharedPtr<IToolkitHost>()) override;
+	virtual bool ShouldForceWorldCentric() override;
+};
diff --git a/Engine/Plugins/MovieScene/ActorAnimationEditor/Source/ActorAnimationEditor/Private/ActorAnimationEditorModule.cpp b/Engine/Plugins/MovieScene/ActorAnimationEditor/Source/ActorAnimationEditor/Private/ActorAnimationEditorModule.cpp
new file mode 100644
--- /dev/null
+++ b/Engine/Plugins/MovieScene/ActorAnimationEditor/Source/ActorAnimationEditor/Private/ActorAnimationEditorModule.cpp
@@ -0,0 +1,72 @@
+// Copyright 1998-2015 Epic Games, Inc. All Rights Reserved.
+
+#include "ActorAnimationEditorPrivatePCH.h"
+#include "ModuleInterface.h"
+
+
+/**
+ * Implements the ActorAnimationEditor module.
+ */
+class FActorAnimationEditorModule
+	: public IModuleInterface
+{
+public:
+
+	// IModuleInterface interface
+
+	virtual void StartupModule() override
+	{
+		RegisterAssetTools();
+	}
+	
+	virtual void ShutdownModule() override
+	{
+		UnregisterAssetTools();
+	}
+
+protected:
+
+	/** Registers asset tool actions. */
+	void RegisterAssetTools()
+	{
+		IAssetTools& AssetTools = FModuleManager::LoadModuleChecked<FAssetToolsModule>("AssetTools").Get();
+
+		RegisterAssetTypeAction(AssetTools, MakeShareable(new FActorAnimationActions));
+	}
+
+	/**
+	 * Registers a single asset type action.
+	 *
+	 * @param AssetTools The asset tools object to register with.
+	 * @param Action The asset type action to register.
+	 */
+	void RegisterAssetTypeAction( IAssetTools& AssetTools, TSharedRef<IAssetTypeActions> Action )
+	{
+		AssetTools.RegisterAssetTypeActions(Action);
+		RegisteredAssetTypeActions.Add(Action);
+	}
+
+	/** Unregisters asset tool actions. */
+	void UnregisterAssetTools()
+	{
+		FAssetToolsModule* AssetToolsModule = FModuleManager::GetModulePtr<FAssetToolsModule>("AssetTools");
+
+		if (AssetToolsModule != nullptr)
+		{
+			IAssetTools& AssetTools = AssetToolsModule->Get();
+
+			for (auto Action : RegisteredAssetTypeActions)
+			{
+				AssetTools.UnregisterAssetTypeActions(Action);
+			}
+		}
+	}
+
+private:
+
+	/** The collection of registered asset type actions. */
+	TArray<TSharedRef<IAssetTypeActions>> RegisteredAssetTypeActions;
+};
+
+
+IMPLEMENT_MODULE(FActorAnimationEditorModule, ActorAnimationEditor);
diff --git a/Engine/Plugins/MovieScene/ActorAnimationEditor/Source/ActorAnimationEditor/Private/ActorAnimationEditorPrivatePCH.h b/Engine/Plugins/MovieScene/ActorAnimationEditor/Source/ActorAnimationEditor/Private/ActorAnimationEditorPrivatePCH.h
new file mode 100644
--- /dev/null
+++ b/Engine/Plugins/MovieScene/ActorAnimationEditor/Source/ActorAnimationEditor/Private/ActorAnimationEditorPrivatePCH.h
@@ -0,0 +1,22 @@
+// Copyright 1998-2015 Epic Games, Inc. All Rights Reserved.
+
+#pragma once
+
+
+/* Private dependencies
+ *****************************************************************************/
+
+#include "ActorAnimation.h"
+#include "CoreUObject.h"
+#include "AssetToolsModule.h"
+#include "EditorStyle.h"
+#include "Engine.h"
+#include "SlateBasics.h"
+#include "SlateStyle.h"
+
+
+/* Private includes
+ *****************************************************************************/
+
+#include "ActorAnimationActions.h"
+#include "ActorAnimationEditorToolkit.h"
diff --git a/Engine/Plugins/MovieScene/ActorAnimationEditor/Source/ActorAnimationEditor/Private/ActorAnimationEditorToolkit.cpp b/Engine/Plugins/MovieScene/ActorAnimationEditor/Source/ActorAnimationEditor/Private/ActorAnimationEditorToolkit.cpp
new file mode 100644
--- /dev/null
+++ b/Engine/Plugins/MovieScene/ActorAnimationEditor/Source/ActorAnimationEditor/Private/ActorAnimationEditorToolkit.cpp
@@ -0,0 +1,210 @@
+// Copyright 1998-2015 Epic Games, Inc. All Rights Reserved.
+
+#include "ActorAnimationEditorPrivatePCH.h"
+#include "Toolkits/IToolkitHost.h"
+#include "Editor/LevelEditor/Public/LevelEditor.h"
+#include "Editor/LevelEditor/Public/ILevelViewport.h"
+#include "Editor/WorkspaceMenuStructure/Public/WorkspaceMenuStructure.h"
+#include "Editor/WorkspaceMenuStructure/Public/WorkspaceMenuStructureModule.h"
+#include "ISequencer.h"
+#include "ISequencerModule.h"
+#include "SDockTab.h"
+
+
+#define LOCTEXT_NAMESPACE "Sequencer"
+
+const FName FActorAnimationEditorToolkit::SequencerMainTabId(TEXT("Sequencer_SequencerMain"));
+
+
+namespace SequencerDefs
+{
+	static const FName SequencerAppIdentifier( TEXT( "SequencerApp" ) );
+}
+
+
+void FActorAnimationEditorToolkit::RegisterTabSpawners(const TSharedRef<class FTabManager>& TabManager)
+{
+	if (IsWorldCentricAssetEditor())
+	{
+		return;
+	}
+
+	WorkspaceMenuCategory = TabManager->AddLocalWorkspaceMenuCategory(LOCTEXT("WorkspaceMenu_SequencerAssetEditor", "Sequencer"));
+
+	TabManager->RegisterTabSpawner(SequencerMainTabId, FOnSpawnTab::CreateSP(this, &FActorAnimationEditorToolkit::SpawnTab_SequencerMain))
+		.SetDisplayName(LOCTEXT("SequencerMainTab", "Sequencer"))
+		.SetGroup(WorkspaceMenuCategory.ToSharedRef());
+}
+
+
+void FActorAnimationEditorToolkit::UnregisterTabSpawners(const TSharedRef<class FTabManager>& TabManager)
+{
+	if (!IsWorldCentricAssetEditor())
+	{
+		TabManager->UnregisterTabSpawner( SequencerMainTabId );
+	}
+	
+	// @todo remove when world-centric mode is added
+	FLevelEditorModule& LevelEditorModule = FModuleManager::LoadModuleChecked<FLevelEditorModule>("LevelEditor");
+	LevelEditorModule.AttachSequencer( SNullWidget::NullWidget, nullptr );
+}
+
+
+void FActorAnimationEditorToolkit::Initialize( const EToolkitMode::Type Mode, const FSequencerViewParams& InViewParams, const TSharedPtr<IToolkitHost>& InitToolkitHost, UActorAnimation* ActorAnimation, bool bEditWithinLevelEditor )
+{
+	{
+		const TSharedRef<FTabManager::FLayout> StandaloneDefaultLayout = FTabManager::NewLayout("Standalone_Sequencer_Layout")
+			->AddArea
+			(
+				FTabManager::NewPrimaryArea()
+					->Split
+					(
+						FTabManager::NewStack()
+						->AddTab(SequencerMainTabId, ETabState::OpenedTab)
+					)
+			);
+
+		const bool bCreateDefaultStandaloneMenu = true;
+		const bool bCreateDefaultToolbar = false;
+		FAssetEditorToolkit::InitAssetEditor(Mode, InitToolkitHost, SequencerDefs::SequencerAppIdentifier, StandaloneDefaultLayout, bCreateDefaultStandaloneMenu, bCreateDefaultToolbar, ActorAnimation);
+	}
+	
+	// Init Sequencer
+	FSequencerInitParams SequencerInitParams;
+	{
+		SequencerInitParams.ViewParams = InViewParams;
+		SequencerInitParams.Animation = ActorAnimation;
+		SequencerInitParams.bEditWithinLevelEditor = bEditWithinLevelEditor;
+		SequencerInitParams.ToolkitHost = InitToolkitHost;
+	}
+
+	Sequencer = FModuleManager::LoadModuleChecked<ISequencerModule>("Sequencer").CreateSequencer(SequencerInitParams);
+
+	if (bEditWithinLevelEditor)
+	{
+		// @todo remove when world-centric mode is added
+		FLevelEditorModule& LevelEditorModule = FModuleManager::LoadModuleChecked<FLevelEditorModule>("LevelEditor");
+
+		LevelEditorModule.AttachSequencer( Sequencer->GetSequencerWidget(), SharedThis(this));
+			
+		// We need to find out when the user loads a new map, because we might need to re-create puppet actors
+		// when previewing a MovieScene
+		LevelEditorModule.OnMapChanged().AddRaw(this, &FActorAnimationEditorToolkit::HandleMapChanged);
+	}
+}
+
+
+FActorAnimationEditorToolkit::FActorAnimationEditorToolkit()
+{
+	// Register sequencer menu extenders.
+	ISequencerModule& SequencerModule = FModuleManager::Get().LoadModuleChecked<ISequencerModule>( "Sequencer" );
+	int32 NewIndex = SequencerModule.GetMenuExtensibilityManager()->GetExtenderDelegates().Add(
+		FAssetEditorExtender::CreateRaw( this, &FActorAnimationEditorToolkit::GetContextSensitiveSequencerExtender ) );
+	SequencerExtenderHandle = SequencerModule.GetMenuExtensibilityManager()->GetExtenderDelegates()[NewIndex].GetHandle();
+}
+
+
+FActorAnimationEditorToolkit::~FActorAnimationEditorToolkit()
+{
+	Sequencer->Close();
+
+	// Unregister delegates
+	if( FModuleManager::Get().IsModuleLoaded( TEXT( "LevelEditor" ) ) )
+	{
+		auto& LevelEditorModule = FModuleManager::LoadModuleChecked< FLevelEditorModule >( TEXT( "LevelEditor" ) );
+		LevelEditorModule.OnMapChanged().RemoveAll( this );
+	}
+
+	// Un-Register sequencer menu extenders.
+	ISequencerModule& SequencerModule = FModuleManager::Get().LoadModuleChecked<ISequencerModule>( "Sequencer" );
+	SequencerModule.GetMenuExtensibilityManager()->GetExtenderDelegates().RemoveAll( [this]( const FAssetEditorExtender& Extender )
+	{
+		return SequencerExtenderHandle == Extender.GetHandle();
+	} );
+}
+
+
+TSharedRef<SDockTab> FActorAnimationEditorToolkit::SpawnTab_SequencerMain(const FSpawnTabArgs& Args)
+{
+	check(Args.GetTabId() == SequencerMainTabId);
+
+	return SNew(SDockTab)
+		.Icon( FEditorStyle::GetBrush("Sequencer.Tabs.SequencerMain") )
+		.Label( LOCTEXT("SequencerMainTitle", "Sequencer") )
+		.TabColorScale( GetTabColorScale() )
+		[
+			Sequencer->GetSequencerWidget()
+		];
+}
+
+
+FName FActorAnimationEditorToolkit::GetToolkitFName() const
+{
+	static FName SequencerName("Sequencer");
+	return SequencerName;
+}
+
+
+FText FActorAnimationEditorToolkit::GetBaseToolkitName() const
+{
+	return LOCTEXT("AppLabel", "Sequencer");
+}
+
+
+FLinearColor FActorAnimationEditorToolkit::GetWorldCentricTabColorScale() const
+{
+	return FLinearColor( 0.7, 0.0f, 0.0f, 0.5f );
+}
+
+
+FString FActorAnimationEditorToolkit::GetWorldCentricTabPrefix() const
+{
+	return LOCTEXT("WorldCentricTabPrefix", "Sequencer ").ToString();
+}
+
+TSharedRef<FExtender> FActorAnimationEditorToolkit::GetContextSensitiveSequencerExtender( const TSharedRef<FUICommandList> CommandList, const TArray<UObject*> ContextSensitiveObjects )
+{
+	TSharedRef<FExtender> AddTrackMenuExtender( new FExtender() );
+	AddTrackMenuExtender->AddMenuExtension(
+		SequencerMenuExtensionPoints::AddTrackMenu_PropertiesSection,
+		EExtensionHook::Before,
+		CommandList,
+		FMenuExtensionDelegate::CreateRaw( this, &FActorAnimationEditorToolkit::ExtendSequencerAddTrackMenu, ContextSensitiveObjects ) );
+	return AddTrackMenuExtender;
+}
+
+void FActorAnimationEditorToolkit::ExtendSequencerAddTrackMenu( FMenuBuilder& AddTrackMenuBuilder, TArray<UObject*> ContextObjects )
+{
+	if ( ContextObjects.Num() == 1 )
+	{
+		AActor* Actor = Cast<AActor>( ContextObjects[0] );
+		if ( Actor != nullptr )
+		{
+			AddTrackMenuBuilder.BeginSection( "Components", LOCTEXT( "ComponentsSection", "Components" ) );
+			{
+				for ( UActorComponent* Component : Actor->GetComponents() )
+				{
+					FUIAction AddComponentAction( FExecuteAction::CreateSP( this, &FActorAnimationEditorToolkit::AddComponentTrack, Component ) );
+					FText AddComponentLabel = FText::FromString(Component->GetName());
+					FText AddComponentToolTip = FText::Format( LOCTEXT( "ComponentToolTipFormat", "Add {0} component" ), FText::FromString( Component->GetName() ) );
+					AddTrackMenuBuilder.AddMenuEntry( AddComponentLabel, AddComponentToolTip, FSlateIcon(), AddComponentAction );
+				}
+			}
+			AddTrackMenuBuilder.EndSection();
+		}
+	}
+}
+
+void FActorAnimationEditorToolkit::AddComponentTrack( UActorComponent* Component )
+{
+	Sequencer->GetHandleToObject( Component );
+}
+
+
+void FActorAnimationEditorToolkit::HandleMapChanged(class UWorld* NewWorld, EMapChangeType MapChangeType)
+{
+	Sequencer->NotifyMapChanged(NewWorld, MapChangeType);
+}
+
+
+#undef LOCTEXT_NAMESPACE
diff --git a/Engine/Plugins/MovieScene/ActorAnimationEditor/Source/ActorAnimationEditor/Private/ActorAnimationEditorToolkit.h b/Engine/Plugins/MovieScene/ActorAnimationEditor/Source/ActorAnimationEditor/Private/ActorAnimationEditorToolkit.h
new file mode 100644
--- /dev/null
+++ b/Engine/Plugins/MovieScene/ActorAnimationEditor/Source/ActorAnimationEditor/Private/ActorAnimationEditorToolkit.h
@@ -0,0 +1,85 @@
+// Copyright 1998-2015 Epic Games, Inc. All Rights Reserved.
+
+#pragma once
+
+#include "Toolkits/AssetEditorToolkit.h"
+
+
+enum class EMapChangeType : uint8;
+class FTabManager;
+class ILevelViewport;
+class ISequencer;
+class IToolkitHost;
+class SWidget;
+class UActorAnimation;
+class UWorld;
+struct FSequencerViewParams;
+
+
+/**
+ * Implements an Editor toolkit for actor animations.
+ */
+class FActorAnimationEditorToolkit
+	: public FAssetEditorToolkit
+{ 
+public:
+
+	/** Default constructor */
+	FActorAnimationEditorToolkit();
+
+	/** Virtual destructor */
+	virtual ~FActorAnimationEditorToolkit();
+
+public:
+
+	/**
+	 * Initialize this asset editor.
+	 *
+	 * @param Mode Asset editing mode for this editor (standalone or world-centric).
+	 * @param InViewParams Parameters for how to view sequencer UI.
+	 * @param InitToolkitHost When Mode is WorldCentric, this is the level editor instance to spawn this editor within.
+	 * @param ActorAnimation The animation to edit.
+	 * @param TrackEditorDelegates Delegates to call to create auto-key handlers for this sequencer.
+	 * @param bEditWithinLevelEditor Whether or not sequencer should be edited within the level editor.
+	 */
+	void Initialize(const EToolkitMode::Type Mode, const FSequencerViewParams& InViewParams, const TSharedPtr<IToolkitHost>& InitToolkitHost, UActorAnimation* ActorAnimation, bool bEditWithinLevelEditor);
+
+public:
+
+	// IToolkit interface
+
+	virtual FName GetToolkitFName() const override;
+	virtual FText GetBaseToolkitName() const override;
+	virtual FString GetWorldCentricTabPrefix() const override;
+	virtual FLinearColor GetWorldCentricTabColorScale() const override;
+	virtual void RegisterTabSpawners(const TSharedRef<FTabManager>& TabManager) override;
+	virtual void UnregisterTabSpawners(const TSharedRef<FTabManager>& TabManager) override;
+
+
+	void UpdateViewports(AActor* ActorToViewThrough) const;
+
+private:
+
+	EVisibility GetTransportControlVisibility(TWeakPtr<ILevelViewport> LevelViewport) const;
+	TSharedRef<SDockTab> SpawnTab_SequencerMain(const FSpawnTabArgs& Args);	
+
+	TSharedRef<FExtender> GetContextSensitiveSequencerExtender(const TSharedRef<FUICommandList> CommandList, const TArray<UObject*> ContextSensitiveObjects);
+	void ExtendSequencerAddTrackMenu(FMenuBuilder& AddTrackMenuBuilder, TArray<UObject*> ContextObjects);
+	void AddComponentTrack(UActorComponent* Component);
+	void HandleMapChanged(UWorld* NewWorld, EMapChangeType MapChangeType);
+
+private:
+
+	/** The sequencer used by this editor. */
+	TSharedPtr<ISequencer> Sequencer;
+
+	/** A map of all the transport controls to viewports that this sequencer has made */
+	TMap<TWeakPtr<ILevelViewport>, TSharedPtr<SWidget>> TransportControls;
+
+	FDelegateHandle SequencerExtenderHandle;
+
+private:
+
+	/**	The tab ids for all the tabs used */
+	static const FName SequencerMainTabId;
+};
diff --git a/Engine/Plugins/MovieScene/ActorAnimationEditor/Source/ActorAnimationEditor/Private/ActorAnimationFactoryNew.cpp b/Engine/Plugins/MovieScene/ActorAnimationEditor/Source/ActorAnimationEditor/Private/ActorAnimationFactoryNew.cpp
new file mode 100644
--- /dev/null
+++ b/Engine/Plugins/MovieScene/ActorAnimationEditor/Source/ActorAnimationEditor/Private/ActorAnimationFactoryNew.cpp
@@ -0,0 +1,46 @@
+// Copyright 1998-2015 Epic Games, Inc. All Rights Reserved.
+
+#include "ActorAnimationEditorPrivatePCH.h"
+#include "ActorAnimation.h"
+#include "ActorAnimationFactoryNew.h"
+
+
+#define LOCTEXT_NAMESPACE "MovieSceneFactory"
+
+
+/* UActorAnimationFactory structors
+ *****************************************************************************/
+
+UActorAnimationFactoryNew::UActorAnimationFactoryNew(const FObjectInitializer& ObjectInitializer)
+	: Super(ObjectInitializer)
+{
+	// @todo sequencer: Only allow users to create new MovieScenes if that feature is turned on globally.
+	if (FParse::Param(FCommandLine::Get(), TEXT("Sequencer")))
+	{
+		bCreateNew = true;
+	}
+
+	bEditAfterNew = true;
+	SupportedClass = UActorAnimation::StaticClass();
+}
+
+
+/* UFactory interface
+ *****************************************************************************/
+
+UObject* UActorAnimationFactoryNew::FactoryCreateNew(UClass* Class, UObject* InParent, FName Name, EObjectFlags Flags, UObject* Context, FFeedbackContext* Warn)
+{
+	auto NewActorAnimation = NewObject<UActorAnimation>(InParent, Name, Flags);
+	NewActorAnimation->Initialize();
+
+	return NewActorAnimation;
+}
+
+
+bool UActorAnimationFactoryNew::ShouldShowInNewMenu() const
+{
+	return true;
+}
+
+
+#undef LOCTEXT_NAMESPACE
\ No newline at end of file
diff --git a/Engine/Plugins/MovieScene/ActorAnimationEditor/Source/ActorAnimationEditor/Private/ActorAnimationFactoryNew.h b/Engine/Plugins/MovieScene/ActorAnimationEditor/Source/ActorAnimationEditor/Private/ActorAnimationFactoryNew.h
new file mode 100644
--- /dev/null
+++ b/Engine/Plugins/MovieScene/ActorAnimationEditor/Source/ActorAnimationEditor/Private/ActorAnimationFactoryNew.h
@@ -0,0 +1,24 @@
+// Copyright 1998-2015 Epic Games, Inc. All Rights Reserved.
+
+#pragma once
+
+#include "Factories/Factory.h"
+#include "ActorAnimationFactoryNew.generated.h"
+
+
+/**
+ * Implements a factory for UActorAnimation objects.
+ */
+UCLASS(hidecategories=Object)
+class UActorAnimationFactoryNew
+	: public UFactory
+{
+	GENERATED_UCLASS_BODY()
+
+public:
+
+	// UFactory Interface
+
+	virtual UObject* FactoryCreateNew(UClass* Class, UObject* InParent, FName Name, EObjectFlags Flags, UObject* Context, FFeedbackContext* Warn) override;
+	virtual bool ShouldShowInNewMenu() const override;
+};
diff --git a/Engine/Plugins/MovieScene/ActorAnimationEditor/Source/ActorAnimationEditor/Private/ActorAnimationPlayer.cpp b/Engine/Plugins/MovieScene/ActorAnimationEditor/Source/ActorAnimationEditor/Private/ActorAnimationPlayer.cpp
new file mode 100644
--- /dev/null
+++ b/Engine/Plugins/MovieScene/ActorAnimationEditor/Source/ActorAnimationEditor/Private/ActorAnimationPlayer.cpp
@@ -0,0 +1,269 @@
+// Copyright 1998-2015 Epic Games, Inc. All Rights Reserved.
+
+#include "ActorAnimationEditorPrivatePCH.h"
+#include "ActorAnimationPlayer.h"
+#include "MovieScene.h"
+#include "MovieSceneAnimation.h"
+#include "MovieSceneInstance.h"
+
+
+/* UActorAnimationPlayer structors
+ *****************************************************************************/
+
+UActorAnimationPlayer::UActorAnimationPlayer(const FObjectInitializer& ObjectInitializer)
+	: Super(ObjectInitializer)
+	, ActorAnimation(nullptr)
+	, bIsPlaying(false)
+	, TimeCursorPosition(0.0f)
+{ }
+
+
+/* UActorAnimationPlayer interface
+ *****************************************************************************/
+
+void UActorAnimationPlayer::PlayActorAnimation(UObject* WorldContextObject, UActorAnimation* ActorAnimation)
+{
+	if (ActorAnimation == nullptr)
+	{
+		return;
+	}
+
+	UWorld* World = GEngine->GetWorldFromContextObject(WorldContextObject);
+
+	check(World != nullptr);
+	check(World->PersistentLevel != nullptr);
+
+	UActorAnimationPlayer* NewPlayer = NewObject<UActorAnimationPlayer>((UObject*)GetTransientPackage(), NAME_None, RF_Transient);
+
+	check(NewPlayer != nullptr);
+
+	NewPlayer->Initialize(ActorAnimation, World);
+	World->PersistentLevel->AddActiveRuntimeMovieScenePlayer(NewPlayer);
+}
+
+
+bool UActorAnimationPlayer::IsPlaying() const
+{
+	return bIsPlaying;
+}
+
+
+void UActorAnimationPlayer::Pause()
+{
+	bIsPlaying = false;
+}
+
+
+void UActorAnimationPlayer::Play()
+{
+	if ((ActorAnimation == nullptr) || !World.IsValid())
+	{
+		return;
+	}
+
+	// @todo Sequencer playback: Should we recreate the instance every time?
+	RootMovieSceneInstance = MakeShareable(new FMovieSceneInstance(*ActorAnimation->GetMovieScene()));
+
+	// @odo Sequencer Should we spawn actors here?
+	SpawnActorsForMovie(RootMovieSceneInstance.ToSharedRef());
+	RootMovieSceneInstance->RefreshInstance(*this);
+	
+	bIsPlaying = true;
+}
+
+
+/* UActorAnimationPlayer implementation
+ *****************************************************************************/
+
+void UActorAnimationPlayer::Initialize(UActorAnimation* InActorAnimation, UWorld* InWorld)
+{
+	ActorAnimation = InActorAnimation;
+	World = InWorld;
+}
+
+
+/* IMovieScenePlayer interface
+ *****************************************************************************/
+
+void UActorAnimationPlayer::SpawnActorsForMovie(TSharedRef<FMovieSceneInstance> MovieSceneInstance)
+{
+	UWorld* WorldPtr = World.Get();
+
+	if (WorldPtr == nullptr)
+	{
+		return;
+	}
+
+	UMovieScene* MovieScene = MovieSceneInstance->GetMovieScene();
+
+	if (MovieScene == nullptr)
+	{
+		return;
+	}
+
+	TArray<FSpawnedActorInfo>* FoundSpawnedActors = InstanceToSpawnedActorMap.Find(MovieSceneInstance);
+
+	if (FoundSpawnedActors != nullptr)
+	{
+		// Remove existing spawned actors for this movie
+		DestroyActorsForMovie( MovieSceneInstance );
+	}
+
+	TArray<FSpawnedActorInfo>& SpawnedActorList = InstanceToSpawnedActorMap.Add(MovieSceneInstance, TArray<FSpawnedActorInfo>());
+
+	for (auto SpawnableIndex = 0; SpawnableIndex < MovieScene->GetSpawnableCount(); ++SpawnableIndex)
+	{
+		auto& Spawnable = MovieScene->GetSpawnable(SpawnableIndex);
+		UClass* GeneratedClass = Spawnable.GetClass();
+		
+		if ((GeneratedClass == nullptr) || !GeneratedClass->IsChildOf(AActor::StaticClass()))
+		{
+			continue;
+		}
+
+		AActor* ActorCDO = CastChecked<AActor>(GeneratedClass->ClassDefaultObject);
+		const FVector SpawnLocation = ActorCDO->GetRootComponent()->RelativeLocation;
+		const FRotator SpawnRotation = ActorCDO->GetRootComponent()->RelativeRotation;
+
+		FActorSpawnParameters SpawnInfo;
+		{
+			SpawnInfo.ObjectFlags = RF_NoFlags;
+		}
+
+		AActor* NewActor = WorldPtr->SpawnActor(GeneratedClass, &SpawnLocation, &SpawnRotation, SpawnInfo);
+
+		if (NewActor)
+		{	
+			// Actor was spawned OK!
+			FSpawnedActorInfo NewInfo;
+			{
+				NewInfo.RuntimeGuid = Spawnable.GetGuid();
+				NewInfo.SpawnedActor = NewActor;
+			}
+						
+			SpawnedActorList.Add(NewInfo);
+		}
+	}
+}
+
+
+void UActorAnimationPlayer::DestroyActorsForMovie(TSharedRef<FMovieSceneInstance> MovieSceneInstance)
+{
+	UWorld* WorldPtr = World.Get();
+
+	if (WorldPtr == nullptr)
+	{
+		return;
+	}
+
+	TArray<FSpawnedActorInfo>* SpawnedActors = InstanceToSpawnedActorMap.Find(MovieSceneInstance);
+
+	if (SpawnedActors == nullptr)
+	{
+		return;
+	}
+
+	TArray<FSpawnedActorInfo>& SpawnedActorsRef = *SpawnedActors;
+
+	for(int32 ActorIndex = 0; ActorIndex < SpawnedActors->Num(); ++ActorIndex)
+	{
+		AActor* Actor = SpawnedActorsRef[ActorIndex].SpawnedActor.Get();
+		
+		if (Actor != nullptr)
+		{
+			// @todo Sequencer figure this out.  Defaults to false.
+			const bool bNetForce = false;
+
+			// At runtime, level modification is not needed
+			const bool bShouldModifyLevel = false;
+			Actor->GetWorld()->DestroyActor(Actor, bNetForce, bShouldModifyLevel);
+		}
+	}
+
+	InstanceToSpawnedActorMap.Remove(MovieSceneInstance);
+}
+
+
+void UActorAnimationPlayer::GetRuntimeObjects(TSharedRef<FMovieSceneInstance> MovieSceneInstance, const FGuid& ObjectId, TArray<UObject*>& OutObjects) const
+{
+	// @todo sequencer runtime: Add support for individually spawning actors on demand when first requested?
+	//    This may be important to reduce the up-front hitch when spawning actors for the entire movie, but
+	//    may introduce smaller hitches during playback.  Needs experimentation.
+
+	const TArray<FSpawnedActorInfo>* SpawnedActors = InstanceToSpawnedActorMap.Find(MovieSceneInstance);
+
+	if ((SpawnedActors != nullptr) && (SpawnedActors->Num() > 0))
+	{
+		// we have a spawned actor for this handle
+		const TArray<FSpawnedActorInfo>& SpawnedActorInfoRef = *SpawnedActors;
+
+		for (int32 SpawnedActorIndex = 0; SpawnedActorIndex < SpawnedActorInfoRef.Num(); ++SpawnedActorIndex)
+		{
+			const FSpawnedActorInfo ActorInfo = SpawnedActorInfoRef[SpawnedActorIndex];
+
+			if ((ActorInfo.RuntimeGuid == ObjectId) && ActorInfo.SpawnedActor.IsValid())
+			{
+				OutObjects.Add(ActorInfo.SpawnedActor.Get());
+			}
+		}
+	}
+	else
+	{
+		// otherwise, check whether we have one or more possessed actors that are mapped to this handle
+		UObject* FoundObject = ActorAnimation->FindObject(ObjectId);
+
+		if (FoundObject != nullptr)
+		{
+			OutObjects.Add(FoundObject);
+		}
+	}
+}
+
+
+void UActorAnimationPlayer::UpdateCameraCut(UObject* ObjectToViewThrough, bool bNewCameraCut) const
+{
+}
+
+
+EMovieScenePlayerStatus::Type UActorAnimationPlayer::GetPlaybackStatus() const
+{
+	return bIsPlaying ? EMovieScenePlayerStatus::Playing : EMovieScenePlayerStatus::Stopped;
+}
+
+
+void UActorAnimationPlayer::AddOrUpdateMovieSceneInstance(UMovieSceneSection& MovieSceneSection, TSharedRef<FMovieSceneInstance> InstanceToAdd)
+{
+	SpawnActorsForMovie( InstanceToAdd );
+}
+
+
+void UActorAnimationPlayer::RemoveMovieSceneInstance(UMovieSceneSection& MovieSceneSection, TSharedRef<FMovieSceneInstance> InstanceToRemove)
+{
+	const bool bDestroyAll = true;
+	DestroyActorsForMovie( InstanceToRemove );
+}
+
+
+TSharedRef<FMovieSceneInstance> UActorAnimationPlayer::GetRootMovieSceneInstance() const
+{
+	return RootMovieSceneInstance.ToSharedRef();
+}
+
+
+/* IRuntimeMovieScenePlayerInterface interface
+ *****************************************************************************/
+
+void UActorAnimationPlayer::Tick(const float DeltaSeconds)
+{
+	float LastTimePosition = TimeCursorPosition;
+
+	if (bIsPlaying)
+	{
+		TimeCursorPosition += DeltaSeconds;
+	}
+
+	if(RootMovieSceneInstance.IsValid())
+	{
+		RootMovieSceneInstance->Update(TimeCursorPosition, LastTimePosition, *this);
+	}
+}
diff --git a/Engine/Plugins/MovieScene/ActorAnimationEditor/Source/ActorAnimationEditor/Private/ActorAnimationPlayer.h b/Engine/Plugins/MovieScene/ActorAnimationEditor/Source/ActorAnimationEditor/Private/ActorAnimationPlayer.h
new file mode 100644
--- /dev/null
+++ b/Engine/Plugins/MovieScene/ActorAnimationEditor/Source/ActorAnimationEditor/Private/ActorAnimationPlayer.h
@@ -0,0 +1,114 @@
+// Copyright 1998-2015 Epic Games, Inc. All Rights Reserved.
+
+#pragma once
+
+#include "IMovieScenePlayer.h"
+#include "ActorAnimation.h"
+#include "IMovieScenePlayer.h"
+#include "Runtime/Engine/Classes/MovieScene/RuntimeMovieScenePlayerInterface.h"
+#include "ActorAnimationPlayer.generated.h"
+
+
+class FMovieSceneInstance;
+class ULevel;
+class UMovieSceneBindings;
+
+
+/**
+ * RuntimeMovieScenePlayer is used to actually "play" a MovieScene asset at runtime.
+ *
+ * This class keeps track of playback state and provides functions for manipulating
+ * a MovieScene while its playing.
+ */
+UCLASS(MinimalAPI)
+class UActorAnimationPlayer
+	: public UObject
+	, public IMovieScenePlayer
+	, public IRuntimeMovieScenePlayerInterface
+{
+	GENERATED_UCLASS_BODY()
+
+public:
+
+	/**
+	 * Play an actor animation.
+	 *
+	 * @param ActorAnimation The actor animation to play.
+	 */
+	UFUNCTION(BlueprintCallable, Category="Game|Cinematic", meta=(WorldContext="WorldContextObject"))
+	static void PlayActorAnimation(UObject* WorldContextObject, UActorAnimation* ActorAnimation);
+
+	/** Start playback from the current time cursor position. */
+	UFUNCTION(BlueprintCallable, Category="Game|Cinematic")
+	void Play();
+
+	/** Pause playback. */
+	UFUNCTION(BlueprintCallable, Category="Game|Cinematic")
+	void Pause();
+
+	/** Check whether the sequence is actively playing. */
+	bool IsPlaying() const;
+
+protected:
+
+	/**
+	 * Initialize the player.
+	 *
+	 * @param InActorAnimation The actor animation to play.
+	 * @param InWorld The world that the animation is played in.
+	 */
+	void Initialize(UActorAnimation* InActorAnimation, UWorld* InWorld);
+
+protected:
+
+	// IMovieScenePlayer interface
+
+	virtual void SpawnActorsForMovie(TSharedRef<FMovieSceneInstance> MovieSceneInstance);
+	virtual void DestroyActorsForMovie(TSharedRef<FMovieSceneInstance> MovieSceneInstance);
+	virtual void GetRuntimeObjects(TSharedRef<FMovieSceneInstance> MovieSceneInstance, const FGuid& ObjectHandle, TArray<UObject*>& OutObjects) const override;
+	virtual void UpdateCameraCut(UObject* ObjectToViewThrough, bool bNewCameraCut) const override;
+	virtual EMovieScenePlayerStatus::Type GetPlaybackStatus() const override;
+	virtual void AddOrUpdateMovieSceneInstance(UMovieSceneSection& MovieSceneSection, TSharedRef<FMovieSceneInstance> InstanceToAdd) override;
+	virtual void RemoveMovieSceneInstance(UMovieSceneSection& MovieSceneSection, TSharedRef<FMovieSceneInstance> InstanceToRemove) override;
+	virtual TSharedRef<FMovieSceneInstance> GetRootMovieSceneInstance() const override;
+
+protected:
+
+	// IRuntimeMovieScenePlayerInterface
+
+	virtual void Tick(const float DeltaSeconds) override;
+
+private:
+
+	/** The actor animation to play. */
+	UPROPERTY()
+	UActorAnimation* ActorAnimation;
+
+	/** Whether we're currently playing. If false, then sequence playback is paused or was never started. */
+	UPROPERTY()
+	bool bIsPlaying;
+
+	/** The current time cursor position within the sequence (in seconds) */
+	UPROPERTY()
+	double TimeCursorPosition;
+
+private:
+
+	struct FSpawnedActorInfo
+	{
+		/** Identifier that maps this actor to a movie scene spawnable */
+		FGuid RuntimeGuid;
+
+		/** The spawned actor */
+		TWeakObjectPtr<AActor> SpawnedActor;
+	};
+
+	/** Maps spawnable GUIDs to their spawned actor in the world */
+	TMap<TWeakPtr<FMovieSceneInstance>, TArray<FSpawnedActorInfo>> InstanceToSpawnedActorMap;
+
+	/** The root movie scene instance to update when playing. */
+	TSharedPtr<FMovieSceneInstance> RootMovieSceneInstance;
+
+	/** The world this player will spawn actors in, if needed */
+	TWeakObjectPtr<UWorld> World;
+};
diff --git a/Engine/Source/Editor/KismetCompiler/Private/KismetCompiler.cpp b/Engine/Source/Editor/KismetCompiler/Private/KismetCompiler.cpp
--- a/Engine/Source/Editor/KismetCompiler/Private/KismetCompiler.cpp
+++ b/Engine/Source/Editor/KismetCompiler/Private/KismetCompiler.cpp
@@ -13,9 +13,6 @@
 #include "Editor/UnrealEd/Public/Kismet2/KismetEditorUtilities.h"
 #include "Editor/UnrealEd/Public/ScriptDisassembler.h"
 #include "Editor/UnrealEd/Public/ComponentTypeRegistry.h"
-#include "K2Node_PlayMovieScene.h"
-#include "RuntimeMovieScenePlayer.h"
-#include "MovieSceneBindings.h"
 #include "Kismet2/Kismet2NameValidators.h"
 #include "UserDefinedStructureCompilerUtils.h"
 #include "K2Node_EnumLiteral.h"
@@ -2069,210 +2066,6 @@ void FKismetCompilerContext::ExpandTimelineNodes(UEdGraph* SourceGraph)
 	}
 }
 
-UEdGraphPin* FKismetCompilerContext::ExpandNodesToAllocateRuntimeMovieScenePlayer( UEdGraph* SourceGraph, UK2Node_PlayMovieScene* PlayMovieSceneNode, ULevel* Level, UK2Node_TemporaryVariable*& OutPlayerVariableNode )
-{
-	// Call URuntimeMovieScenePlayer::CreatePlayer() to create a new RuntimeMovieScenePlayer instance
-	UK2Node_CallFunction* CreatePlayerCallNode = SpawnIntermediateNode<UK2Node_CallFunction>( PlayMovieSceneNode, SourceGraph );
-	{
-		CreatePlayerCallNode->FunctionReference.SetExternalMember(GET_FUNCTION_NAME_CHECKED(URuntimeMovieScenePlayer, CreatePlayer), URuntimeMovieScenePlayer::StaticClass());
-		CreatePlayerCallNode->AllocateDefaultPins();
-	}
-
-	// The return value of URuntimeMovieScenePlayer::CreatePlayer() is the actual MovieScenePlayer we'll be operating with
-	UEdGraphPin* CreatePlayerReturnValuePin = CreatePlayerCallNode->GetReturnValuePin();
-
-	// Make a literal for the level and bind it to our function call as a parameter
-	UK2Node_Literal* LevelLiteralNode = SpawnIntermediateNode<UK2Node_Literal>( PlayMovieSceneNode, SourceGraph );
-	LevelLiteralNode->AllocateDefaultPins();
-
-	// Make a literal for the MovieSceneBindings object and bind it to our function call as a parameter
-	UK2Node_Literal* MovieSceneBindingsLiteralNode = SpawnIntermediateNode<UK2Node_Literal>( PlayMovieSceneNode, SourceGraph );
-	MovieSceneBindingsLiteralNode->AllocateDefaultPins();
-
-
-	// Create a local variable to store the URuntimeMovieScenePlayer object instance in
-	UK2Node_TemporaryVariable* PlayerVariableNode;
-	{
-		const bool bIsArray = false;
-		PlayerVariableNode = SpawnInternalVariable(
-			PlayMovieSceneNode,
-			CreatePlayerReturnValuePin->PinType.PinCategory, 
-			CreatePlayerReturnValuePin->PinType.PinSubCategory, 
-			CreatePlayerReturnValuePin->PinType.PinSubCategoryObject.Get(),
-			bIsArray );
-	}
-	UEdGraphPin* PlayerVariablePin = PlayerVariableNode->GetVariablePin();
-
-	UK2Node_AssignmentStatement* AssignResultToPlayerVariableNode = SpawnIntermediateNode<UK2Node_AssignmentStatement>( PlayMovieSceneNode, SourceGraph );
-	AssignResultToPlayerVariableNode->AllocateDefaultPins();
-
-	// Create a node that checks to see if our variable that contains the RuntimeMovieScenePlayer instance is null.
-	// If it's null, we'll allocate it now and store it in the variable.
-	UK2Node_CallFunction* ComparisonNode = SpawnIntermediateNode<UK2Node_CallFunction>( PlayMovieSceneNode, SourceGraph );
-	{
-		ComparisonNode->FunctionReference.SetExternalMember(GET_FUNCTION_NAME_CHECKED(UKismetMathLibrary, EqualEqual_ObjectObject), UKismetMathLibrary::StaticClass());
-		ComparisonNode->AllocateDefaultPins();
-	}
-
-	UK2Node_IfThenElse* IfVariableNullNode = SpawnIntermediateNode<UK2Node_IfThenElse>( PlayMovieSceneNode, SourceGraph );
-	{
-		IfVariableNullNode->AllocateDefaultPins();
-	}
-
-
-	// OK, all of our nodes are created.  Now wire everything together!
-	{
-		// The first thing we'll do is check to see if we've allocated a player yet.  If not, we need to
-		// do that now.  So we'll use an "if-then" node to check.
-
-		// Hook the "if-then" node to our comparison function, that simply checks for null
-		IfVariableNullNode->GetConditionPin()->MakeLinkTo( ComparisonNode->GetReturnValuePin() );
-
-		// We'll compare the player variable ('A')...
-		ComparisonNode->FindPinChecked(TEXT("A"))->MakeLinkTo( PlayerVariablePin );
-
-		// ...against a NULL value ('B')
-		ComparisonNode->FindPinChecked(TEXT("B"))->DefaultObject = NULL;
-
-		// If the comparison returns true (variable is null), then we need to call our function to
-		// create the player object
-		IfVariableNullNode->GetThenPin()->MakeLinkTo( CreatePlayerCallNode->GetExecPin() );
-
-		// Setup function params for "URuntimeMovieScenePlayer::CreatePlayer()"
-		{
-			// Our level literal just points to the level object
-			LevelLiteralNode->SetObjectRef( Level );
-
-			{
-				// Duplicate the bindings and store a copy into the level.  We want the bindings to be
-				// outered to the level so they'll be duplicated when the level is duplicated (e.g. for PIE)
-				UMovieSceneBindings* NodeMovieSceneBindings = PlayMovieSceneNode->GetMovieSceneBindings();
-				UMovieSceneBindings* LevelMovieSceneBindings = NULL;
-				if( NodeMovieSceneBindings != NULL && Level != NULL )
-				{
-					LevelMovieSceneBindings = DuplicateObject( NodeMovieSceneBindings, Level );
-					check( LevelMovieSceneBindings != NULL );
-
-					// Tell the Level about the new bindings object.
-					Level->AddMovieSceneBindings( LevelMovieSceneBindings );
-				}
-
-				MovieSceneBindingsLiteralNode->SetObjectRef( LevelMovieSceneBindings );
-			}
-
-			CreatePlayerCallNode->FindPinChecked( TEXT( "Level" ) )->MakeLinkTo( LevelLiteralNode->GetValuePin() );
-			CreatePlayerCallNode->FindPinChecked( TEXT( "MovieSceneBindings" ) )->MakeLinkTo( MovieSceneBindingsLiteralNode->GetValuePin() );
-		}
-
-
-		// Our function that creates the player returns the newly-created player object.  We'll
-		// store that in our variable
-		CreatePlayerCallNode->GetThenPin()->MakeLinkTo( AssignResultToPlayerVariableNode->GetExecPin() );
-		AssignResultToPlayerVariableNode->GetVariablePin()->MakeLinkTo( PlayerVariablePin );
-		AssignResultToPlayerVariableNode->PinConnectionListChanged( AssignResultToPlayerVariableNode->GetVariablePin() );
-		CreatePlayerReturnValuePin->MakeLinkTo( AssignResultToPlayerVariableNode->GetValuePin() );
-		AssignResultToPlayerVariableNode->PinConnectionListChanged( AssignResultToPlayerVariableNode->GetValuePin() );
-	}
-
-	OutPlayerVariableNode = PlayerVariableNode;
-	return IfVariableNullNode->GetExecPin();
-}
-
-
-void FKismetCompilerContext::ExpandPlayMovieSceneNodes( UEdGraph* SourceGraph )
-{
-	ULevel* Level = NULL;
-	if( Blueprint->IsA( ULevelScriptBlueprint::StaticClass() ) )
-	{
-		ULevelScriptBlueprint* LSB = CastChecked< ULevelScriptBlueprint >( Blueprint );
-		Level = LSB->GetLevel();
-	}
-
-	// Wipe old MovieSceneBindings on the Level.  We'll recreate them all now.
-	if( Level != NULL )
-	{
-		Level->ClearMovieSceneBindings();
-	}
-
-	for (int32 ChildIndex = 0; ChildIndex < SourceGraph->Nodes.Num(); ++ChildIndex)
-	{
-		UK2Node_PlayMovieScene* PlayMovieSceneNode = Cast<UK2Node_PlayMovieScene>( SourceGraph->Nodes[ChildIndex] );
-		if( PlayMovieSceneNode != NULL )
-		{
-			UEdGraphPin* PlayPin = PlayMovieSceneNode->GetPlayPin();
-			const bool bPlayPinConnected = PlayPin->LinkedTo.Num() > 0;
-			UEdGraphPin* PausePin = PlayMovieSceneNode->GetPausePin();
-			const bool bPausePinConnected = PausePin->LinkedTo.Num() > 0;
-
-			// Do we need to create a MovieScenePlayer?
-			const bool bNeedMovieScenePlayer = bPlayPinConnected || bPausePinConnected;
-			if( bNeedMovieScenePlayer )
-			{
-				// Generate a node network to allocate a MovieScenePlayer on demand.  All of the various input exec pins
-				// will first be routed through this network, to make sure that we have a movie scene player to work with!
-				UK2Node_TemporaryVariable* PlayerVariableNode = NULL;
-				UEdGraphPin* AllocateRuntimeMovieScenePlayerExecPin = 
-					ExpandNodesToAllocateRuntimeMovieScenePlayer( SourceGraph, PlayMovieSceneNode, Level, PlayerVariableNode );
-				UEdGraphPin* PlayerVariablePin = PlayerVariableNode->GetVariablePin();
-
-
-				// Create a call function node to call 'Play' on the RuntimeMovieScenePlayer object
-				if( bPlayPinConnected )
-				{
-					UK2Node_CallFunction* PlayCallNode = SpawnIntermediateNode<UK2Node_CallFunction>( PlayMovieSceneNode, SourceGraph );
-					{
-						PlayCallNode->FunctionReference.SetExternalMember(GET_FUNCTION_NAME_CHECKED(URuntimeMovieScenePlayer, Play), URuntimeMovieScenePlayer::StaticClass());
-						PlayCallNode->AllocateDefaultPins();
-					}
-					
-					UK2Node_ExecutionSequence* SequenceNode = SpawnIntermediateNode<UK2Node_ExecutionSequence>( PlayMovieSceneNode, SourceGraph );
-					SequenceNode->AllocateDefaultPins();
-				
-					// Move input links from 'Play' to the exec pin on the Sequence node
-					MovePinLinksToIntermediate( *PlayPin, *SequenceNode->GetExecPin() );
-
-					SequenceNode->GetThenPinGivenIndex( 0 )->MakeLinkTo( AllocateRuntimeMovieScenePlayerExecPin );
-
-					// Tell the 'Play' node which player object it's calling the function on
-					UEdGraphPin* PlaySelfPin = Schema->FindSelfPin( *PlayCallNode, EGPD_Input );
-					PlaySelfPin->MakeLinkTo( PlayerVariablePin );
-
-					// Hook our sequence up to call the function
-					UEdGraphPin* PlayExecPin = Schema->FindExecutionPin( *PlayCallNode, EGPD_Input );
-					SequenceNode->GetThenPinGivenIndex( 1 )->MakeLinkTo( PlayExecPin );
-				}
-					
-					
-				// Create a call function node to call 'Pause' on the RuntimeMovieScenePlayer object
-				if( bPausePinConnected )
-				{
-					UK2Node_CallFunction* PauseCallNode = SpawnIntermediateNode<UK2Node_CallFunction>( PlayMovieSceneNode, SourceGraph );
-					{
-						PauseCallNode->FunctionReference.SetExternalMember(GET_FUNCTION_NAME_CHECKED(URuntimeMovieScenePlayer, Pause), URuntimeMovieScenePlayer::StaticClass());
-						PauseCallNode->AllocateDefaultPins();
-					}
-
-					UK2Node_ExecutionSequence* SequenceNode = SpawnIntermediateNode<UK2Node_ExecutionSequence>( PlayMovieSceneNode, SourceGraph );
-					SequenceNode->AllocateDefaultPins();
-				
-					// Move input links from 'Pause' to the exec pin on the Sequence node
-					MovePinLinksToIntermediate( *PausePin, *SequenceNode->GetExecPin() );
-
-					SequenceNode->GetThenPinGivenIndex( 0 )->MakeLinkTo( AllocateRuntimeMovieScenePlayerExecPin );
-
-					// Tell the 'Pause' node which player object it's calling the function on
-					UEdGraphPin* PauseSelfPin = Schema->FindSelfPin( *PauseCallNode, EGPD_Input );
-					PauseSelfPin->MakeLinkTo( PlayerVariablePin );
-
-					// Hook our sequence up to call the function
-					UEdGraphPin* PauseExecPin = Schema->FindExecutionPin( *PauseCallNode, EGPD_Input );
-					SequenceNode->GetThenPinGivenIndex( 1 )->MakeLinkTo( PauseExecPin );
-				}
-			}
-		}
-	}
-}
-
 FPinConnectionResponse FKismetCompilerContext::MovePinLinksToIntermediate(UEdGraphPin& SourcePin, UEdGraphPin& IntermediatePin)
 {
 	 UEdGraphSchema_K2 const* K2Schema = GetSchema();
@@ -2552,9 +2345,6 @@ void FKismetCompilerContext::ExpansionStep(UEdGraph* Graph, bool bAllowUbergraph
 		{
 			// Expand timeline nodes
 			ExpandTimelineNodes(Graph);
-
-			// Expand PlayMovieScene nodes
-			ExpandPlayMovieSceneNodes(Graph);
 		}
 	}
 }
diff --git a/Engine/Source/Editor/KismetCompiler/Public/KismetCompiler.h b/Engine/Source/Editor/KismetCompiler/Public/KismetCompiler.h
--- a/Engine/Source/Editor/KismetCompiler/Public/KismetCompiler.h
+++ b/Engine/Source/Editor/KismetCompiler/Public/KismetCompiler.h
@@ -310,12 +310,6 @@ protected:
 	/** Expand timeline nodes into necessary nodes */
 	void ExpandTimelineNodes(UEdGraph* SourceGraph);
 
-	/** Expand any PlayMovieScene nodes */
-	void ExpandPlayMovieSceneNodes(UEdGraph* SourceGraph);
-
-	/** Used internally by ExpandPlayMovieSceneNodes() to generate a node network to allocate a URuntimeMovieScenePlayer object instance on demand */
-	UEdGraphPin* ExpandNodesToAllocateRuntimeMovieScenePlayer( UEdGraph* SourceGraph, class UK2Node_PlayMovieScene* PlayMovieSceneNode, class ULevel* Level, class UK2Node_TemporaryVariable*& OutPlayerVariableNode );
-
 	/**
 	 * First phase of compiling a function graph
 	 *   - Performs initial validation that the graph is at least well formed enough to be processed further
diff --git a/Engine/Source/Editor/MovieSceneTools/Private/AssetTypeActions_MovieScene.cpp b/Engine/Source/Editor/MovieSceneTools/Private/AssetTypeActions_MovieScene.cpp
deleted file mode 100644
diff --git a/Engine/Source/Editor/MovieSceneTools/Private/AssetTypeActions_MovieScene.h b/Engine/Source/Editor/MovieSceneTools/Private/AssetTypeActions_MovieScene.h
deleted file mode 100644
diff --git a/Engine/Source/Editor/MovieSceneTools/Private/K2Node_PlayMovieScene.cpp b/Engine/Source/Editor/MovieSceneTools/Private/K2Node_PlayMovieScene.cpp
deleted file mode 100644
diff --git a/Engine/Source/Editor/MovieSceneTools/Private/MovieSceneFactory.cpp b/Engine/Source/Editor/MovieSceneTools/Private/MovieSceneFactory.cpp
deleted file mode 100644
diff --git a/Engine/Source/Editor/MovieSceneTools/Private/MovieSceneToolsModule.cpp b/Engine/Source/Editor/MovieSceneTools/Private/MovieSceneToolsModule.cpp
--- a/Engine/Source/Editor/MovieSceneTools/Private/MovieSceneToolsModule.cpp
+++ b/Engine/Source/Editor/MovieSceneTools/Private/MovieSceneToolsModule.cpp
@@ -2,17 +2,10 @@
 
 #include "MovieSceneToolsPrivatePCH.h"
 #include "ModuleManager.h"
-
-#include "MovieSceneFactory.h"
-#include "K2Node_PlayMovieScene.h"
-
 #include "ISequencerModule.h"
 #include "Kismet2/BlueprintEditorUtils.h"
-
 #include "MovieSceneSection.h"
 #include "ISequencerSection.h"
-#include "AssetToolsModule.h"
-#include "AssetTypeActions_MovieScene.h"
 #include "ScopedTransaction.h"
 #include "MovieScene.h"
 #include "MovieSceneTrackEditor.h"
@@ -25,17 +18,14 @@
 #include "ParticleTrackEditor.h"
 #include "PathTrackEditor.h"
 
+
 /**
- * MovieSceneTools module implementation (private)
+ * Implements the MovieSceneTools module.
  */
 class FMovieSceneToolsModule : public IMovieSceneTools
 {
-
 	virtual void StartupModule() override
 	{
-		MovieSceneAssetTypeActions = MakeShareable( new FAssetTypeActions_MovieScene );
-		FModuleManager::LoadModuleChecked< FAssetToolsModule >( "AssetTools" ).Get().RegisterAssetTypeActions( MovieSceneAssetTypeActions.ToSharedRef() );
-
 		// Register with the sequencer module that we provide auto-key handlers.
 		ISequencerModule& SequencerModule = FModuleManager::Get().LoadModuleChecked<ISequencerModule>( "Sequencer" );
 		PropertyTrackEditorCreateTrackEditorDelegateHandle      = SequencerModule.RegisterTrackEditor_Handle( FOnCreateTrackEditor::CreateStatic( &FPropertyTrackEditor::CreateTrackEditor ) );
@@ -48,18 +38,8 @@ class FMovieSceneToolsModule : public IMovieSceneTools
 		PathTrackEditorCreateTrackEditorDelegateHandle          = SequencerModule.RegisterTrackEditor_Handle( FOnCreateTrackEditor::CreateStatic( &F3DPathTrackEditor::CreateTrackEditor ) );
 	}
 
-
 	virtual void ShutdownModule() override
 	{
-		// Only unregister if the asset tools module is loaded.  We don't want to forcibly load it during shutdown phase.
-		check( MovieSceneAssetTypeActions.IsValid() );
-		if( FModuleManager::Get().IsModuleLoaded( "AssetTools" ) )
-		{
-			FModuleManager::GetModuleChecked< FAssetToolsModule >( "AssetTools" ).Get().UnregisterAssetTypeActions( MovieSceneAssetTypeActions.ToSharedRef() );
-		}
-		MovieSceneAssetTypeActions.Reset();
-
-
 		if( FModuleManager::Get().IsModuleLoaded( "Sequencer" ) )
 		{
 			// Unregister auto key handlers
@@ -76,12 +56,8 @@ class FMovieSceneToolsModule : public IMovieSceneTools
 		
 	}
 
-
 private:
 
-	/** Asset type actions for MovieScene assets.  Cached here so that we can unregister it during shutdown. */
-	TSharedPtr< FAssetTypeActions_MovieScene > MovieSceneAssetTypeActions;
-
 	/** Registered delegate handles */
 	FDelegateHandle PropertyTrackEditorCreateTrackEditorDelegateHandle;
 	FDelegateHandle TransformTrackEditorCreateTrackEditorDelegateHandle;
@@ -95,4 +71,3 @@ private:
 
 
 IMPLEMENT_MODULE( FMovieSceneToolsModule, MovieSceneTools );
-
diff --git a/Engine/Source/Editor/MovieSceneTools/Private/TrackEditors/PathTrackEditor.cpp b/Engine/Source/Editor/MovieSceneTools/Private/TrackEditors/PathTrackEditor.cpp
--- a/Engine/Source/Editor/MovieSceneTools/Private/TrackEditors/PathTrackEditor.cpp
+++ b/Engine/Source/Editor/MovieSceneTools/Private/TrackEditors/PathTrackEditor.cpp
@@ -6,7 +6,6 @@
 #include "MovieSceneSection.h"
 #include "MovieScene3DPathTrack.h"
 #include "MovieScene3DPathSection.h"
-#include "ISequencerObjectBindingManager.h"
 #include "ISequencerSection.h"
 #include "ISectionLayoutBuilder.h"
 #include "IKeyArea.h"
diff --git a/Engine/Source/Editor/MovieSceneTools/Private/TrackEditors/PropertyTrackEditor.cpp b/Engine/Source/Editor/MovieSceneTools/Private/TrackEditors/PropertyTrackEditor.cpp
--- a/Engine/Source/Editor/MovieSceneTools/Private/TrackEditors/PropertyTrackEditor.cpp
+++ b/Engine/Source/Editor/MovieSceneTools/Private/TrackEditors/PropertyTrackEditor.cpp
@@ -3,6 +3,7 @@
 #include "MovieSceneToolsPrivatePCH.h"
 #include "PropertyEditing.h"
 #include "MovieScene.h"
+#include "MovieSceneAnimation.h"
 #include "MovieSceneTrack.h"
 #include "MovieSceneFloatTrack.h"
 #include "MovieSceneVectorTrack.h"
@@ -12,7 +13,6 @@
 #include "MovieSceneVisibilityTrack.h"
 #include "ScopedTransaction.h"
 #include "MovieSceneSection.h"
-#include "ISequencerObjectBindingManager.h"
 #include "ISequencerObjectChangeListener.h"
 #include "ISequencerSection.h"
 #include "ISectionLayoutBuilder.h"
@@ -119,10 +119,9 @@ private:
 			{
 				if (MovieSceneBinding.GetTracks()[TrackIndex] == Track)
 				{
-					TArray<UObject*> RuntimeObjects;
-					Sequencer->GetObjectBindingManager()->GetRuntimeObjects(Sequencer->GetFocusedMovieSceneInstance(), MovieSceneBinding.GetObjectGuid(), RuntimeObjects);
+					UObject* RuntimeObject = Sequencer->GetAnimation()->FindObject(MovieSceneBinding.GetObjectGuid());
 
-					for (UObject* RuntimeObject : RuntimeObjects)
+					if (RuntimeObject != nullptr)
 					{
 						UProperty* Property = RuntimeObject->GetClass()->FindPropertyByName(CastChecked<UMovieSceneColorTrack>(Track)->GetPropertyName());
 						UStructProperty* ColorStructProp = Cast<UStructProperty>( Property );
@@ -466,11 +465,10 @@ void FPropertyTrackEditor::AddKey(const FGuid& ObjectGuid, UObject* AdditionalAs
 
 UEnum* GetEnumForByteTrack(TSharedPtr<ISequencer> Sequencer, const FGuid& OwnerObjectHandle, FName PropertyName, UMovieSceneByteTrack* ByteTrack)
 {
-	TArray<UObject*> RuntimeObjects;
-	Sequencer->GetObjectBindingManager()->GetRuntimeObjects(Sequencer->GetFocusedMovieSceneInstance(), OwnerObjectHandle, RuntimeObjects);
+	UObject* RuntimeObject = Sequencer->GetAnimation()->FindObject(OwnerObjectHandle);
 
 	TSet<UEnum*> PropertyEnums;
-	for (UObject* RuntimeObject : RuntimeObjects)
+	if (RuntimeObject != nullptr)
 	{
 		UProperty* Property = RuntimeObject->GetClass()->FindPropertyByName(PropertyName);
 		if (Property != nullptr)
diff --git a/Engine/Source/Editor/MovieSceneTools/Public/K2Node_PlayMovieScene.h b/Engine/Source/Editor/MovieSceneTools/Public/K2Node_PlayMovieScene.h
deleted file mode 100644
diff --git a/Engine/Source/Editor/MovieSceneTools/Public/MovieSceneFactory.h b/Engine/Source/Editor/MovieSceneTools/Public/MovieSceneFactory.h
deleted file mode 100644
diff --git a/Engine/Source/Editor/NiagaraEditor/Private/NiagaraEffectEditor.cpp b/Engine/Source/Editor/NiagaraEditor/Private/NiagaraEffectEditor.cpp
--- a/Engine/Source/Editor/NiagaraEditor/Private/NiagaraEffectEditor.cpp
+++ b/Engine/Source/Editor/NiagaraEditor/Private/NiagaraEffectEditor.cpp
@@ -2,6 +2,7 @@
 
 #include "NiagaraEditorPrivatePCH.h"
 #include "NiagaraEffect.h"
+#include "NiagaraAnimation.h"
 
 #include "Toolkits/IToolkitHost.h"
 #include "Editor/WorkspaceMenuStructure/Public/WorkspaceMenuStructureModule.h"
@@ -81,29 +82,34 @@ void FNiagaraEffectEditor::InitNiagaraEffectEditor(const EToolkitMode::Type Mode
 
 	if (!Sequencer.IsValid())
 	{
-		FSequencerViewParams ViewParams( TEXT( "NiagaraSequencerSettings" ) );
-		ViewParams.InitalViewRange = TRange<float>(-0.02f, 3.2f);
-		ViewParams.InitialScrubPosition = 0;
+		MovieScene = NewObject<UMovieScene>(InEffect, FName("Niagara Effect MovieScene"), RF_RootSet);
+		auto NewAnimation = NewObject<UNiagaraAnimation>(MovieScene);
 
-		SequencerBindingManager = MakeShareable(new FNiagaraSequencerObjectBindingManager());
+		FSequencerViewParams ViewParams(TEXT("NiagaraSequencerSettings"));
+		{
+			ViewParams.InitalViewRange = TRange<float>(-0.02f, 3.2f);
+			ViewParams.InitialScrubPosition = 0;
+		}
 
-		MovieScene = NewObject<UMovieScene>(InEffect, FName("Niagara Effect MovieScene"), RF_RootSet);
+		FSequencerInitParams SequencerInitParams;
+		{
+			SequencerInitParams.ViewParams = ViewParams;
+			SequencerInitParams.Animation = NewAnimation;
+			SequencerInitParams.bEditWithinLevelEditor = false;
+			SequencerInitParams.ToolkitHost = nullptr;
+		}
 
 		ISequencerModule &SeqModule = FModuleManager::LoadModuleChecked< ISequencerModule >("Sequencer");
 		FDelegateHandle CreateTrackEditorHandle = SeqModule.RegisterTrackEditor_Handle(FOnCreateTrackEditor::CreateStatic(&FNiagaraEffectEditor::CreateTrackEditor));
-		Sequencer = SeqModule.CreateSequencer(MovieScene, ViewParams, SequencerBindingManager.ToSharedRef());
+		Sequencer = SeqModule.CreateSequencer(SequencerInitParams);
 
 		for (TSharedPtr<FNiagaraSimulation> Emitter : EffectInstance->GetEmitters())
 		{
 			UEmitterMovieSceneTrack *Track = Cast<UEmitterMovieSceneTrack> (MovieScene->AddMasterTrack(UEmitterMovieSceneTrack::StaticClass()) );
 			 Track->SetEmitter(Emitter);
 		}
-
 	}
 
-
-
-
 	TSharedRef<FTabManager::FLayout> StandaloneDefaultLayout = FTabManager::NewLayout("Standalone_Niagara_Effect_Layout_v7")
 		->AddArea
 		(
diff --git a/Engine/Source/Editor/NiagaraEditor/Private/NiagaraEffectEditor.h b/Engine/Source/Editor/NiagaraEditor/Private/NiagaraEffectEditor.h
--- a/Engine/Source/Editor/NiagaraEditor/Private/NiagaraEffectEditor.h
+++ b/Engine/Source/Editor/NiagaraEditor/Private/NiagaraEffectEditor.h
@@ -78,7 +78,6 @@ private:
 
 	/* stuff needed by the Sequencer */
 	UMovieScene *MovieScene;
-	TSharedPtr<ISequencerObjectBindingManager> SequencerBindingManager;
 	TSharedPtr<ISequencer> Sequencer;
 
 
diff --git a/Engine/Source/Editor/NiagaraEditor/Private/NiagaraSequencer.h b/Engine/Source/Editor/NiagaraEditor/Private/NiagaraSequencer.h
--- a/Engine/Source/Editor/NiagaraEditor/Private/NiagaraSequencer.h
+++ b/Engine/Source/Editor/NiagaraEditor/Private/NiagaraSequencer.h
@@ -2,7 +2,6 @@
 
 #pragma once
 #include "Editor/Sequencer/Public/ISequencerModule.h"
-#include "Editor/Sequencer/Public/ISequencerObjectBindingManager.h"
 #include "Editor/Sequencer/Public/MovieSceneTrackEditor.h"
 #include "Editor/Sequencer/Public/ISequencerSection.h"
 #include "Runtime/MovieScene/Public/MovieScene.h"
@@ -218,36 +217,3 @@ public:
 	}
 
 };
-
-
-
-/**
- * Sequencer object binding manager for Niagara effects; all of this is just stubbed out
-*/
-class FNiagaraSequencerObjectBindingManager : public FGCObject, public ISequencerObjectBindingManager
-{
-public:
-	FNiagaraSequencerObjectBindingManager()
-	{}
-
-	~FNiagaraSequencerObjectBindingManager()
-	{}
-
-	/** FGCObject interface */
-	virtual void AddReferencedObjects(FReferenceCollector& Collector) override {	}
-
-	/** ISequencerObjectBindingManager interface */
-	virtual bool AllowsSpawnableObjects() const override { return false; }
-	virtual FGuid FindGuidForObject(const UMovieScene& MovieScene, UObject& Object) const override { return FGuid(); }
-	virtual void SpawnOrDestroyObjectsForInstance(TSharedRef<FMovieSceneInstance> MovieSceneInstance, bool bDestroyAll) override { }
-	virtual void RemoveMovieSceneInstance( TSharedRef<FMovieSceneInstance> MovieSceneInstance ) override {}
-	virtual void DestroyAllSpawnedObjects() override { }
-	virtual bool CanPossessObject(UObject& Object) const override 	{ return false; }
-	virtual void BindPossessableObject(const FGuid& PossessableGuid, UObject& PossessedObject) override	{};
-	virtual void UnbindPossessableObjects(const FGuid& PossessableGuid) override	{};
-	virtual void GetRuntimeObjects(const TSharedRef<FMovieSceneInstance>& MovieSceneInstance, const FGuid& ObjectGuid, TArray<UObject*>& OutRuntimeObjects) const override	{}
-	virtual UObject* GetParentObject( UObject* Object ) const override { return nullptr; }
-
-	virtual bool TryGetObjectBindingDisplayName(const TSharedRef<FMovieSceneInstance>& MovieSceneInstance, const FGuid& ObjectGuid, FText& DisplayName) const override	{ return false; }
-private:
-};
\ No newline at end of file
diff --git a/Engine/Source/Editor/Sequencer/Private/SSequencer.cpp b/Engine/Source/Editor/Sequencer/Private/SSequencer.cpp
--- a/Engine/Source/Editor/Sequencer/Private/SSequencer.cpp
+++ b/Engine/Source/Editor/Sequencer/Private/SSequencer.cpp
@@ -20,7 +20,6 @@
 #include "DragAndDrop/ActorDragDropGraphEdOp.h"
 #include "DragAndDrop/ClassDragDropOp.h"
 #include "AssetSelection.h"
-#include "K2Node_PlayMovieScene.h"
 #include "MovieSceneShotSection.h"
 #include "CommonMovieSceneTools.h"
 #include "SSearchBox.h"
diff --git a/Engine/Source/Editor/Sequencer/Private/Sequencer.cpp b/Engine/Source/Editor/Sequencer/Private/Sequencer.cpp
--- a/Engine/Source/Editor/Sequencer/Private/Sequencer.cpp
+++ b/Engine/Source/Editor/Sequencer/Private/Sequencer.cpp
@@ -3,7 +3,7 @@
 #include "SequencerPrivatePCH.h"
 #include "Sequencer.h"
 #include "SequencerEdMode.h"
-#include "ISequencerObjectBindingManager.h"
+#include "MovieSceneAnimation.h"
 #include "MovieScene.h"
 #include "Engine/LevelScriptBlueprint.h"
 #include "Editor/PropertyEditor/Public/PropertyEditorModule.h"
@@ -25,7 +25,6 @@
 #include "MovieSceneTrackEditor.h"
 #include "Kismet2/KismetEditorUtilities.h"
 #include "Kismet2/BlueprintEditorUtils.h"
-#include "K2Node_PlayMovieScene.h"
 #include "AssetSelection.h"
 #include "ScopedTransaction.h"
 #include "MovieSceneShotSection.h"
@@ -55,7 +54,7 @@ namespace
 	TSharedPtr<FSequencer> ActiveSequencer;
 }
 
-void FSequencer::InitSequencer( const FSequencerInitParams& InitParams, const TArray<FOnCreateTrackEditor>& TrackEditorDelegates )
+void FSequencer::InitSequencer(const FSequencerInitParams& InitParams, const TSharedRef<ISequencerObjectChangeListener>& InObjectChangeListener, const TArray<FOnCreateTrackEditor>& TrackEditorDelegates)
 {
 	if( IsSequencerEnabled() )
 	{
@@ -81,13 +80,13 @@ void FSequencer::InitSequencer( const FSequencerInitParams& InitParams, const TA
 
 		LastViewRange = TargetViewRange = InitParams.ViewParams.InitalViewRange;
 		ScrubPosition = InitParams.ViewParams.InitialScrubPosition;
-		ObjectChangeListener = InitParams.ObjectChangeListener;
-		ObjectBindingManager = InitParams.ObjectBindingManager;
+		ObjectChangeListener = InObjectChangeListener;
+		Animation = InitParams.Animation;
 
 		check( ObjectChangeListener.IsValid() );
-		check( ObjectBindingManager.IsValid() );
+		check( Animation != nullptr );
 		
-		UMovieScene& RootMovieScene = *InitParams.RootMovieScene;
+		UMovieScene& RootMovieScene = *InitParams.Animation->GetMovieScene();
 		RootMovieScene.SetFlags( RF_Transactional );
 		
 		// Focusing the initial movie scene needs to be done before the first time GetFocusedMovieSceneInstane or GetRootMovieSceneInstance is used
@@ -182,7 +181,7 @@ FSequencer::~FSequencer()
 
 	DestroySpawnablesForAllMovieScenes();
 
-	ObjectBindingManager.Reset();
+	Animation = nullptr;
 	TrackEditors.Empty();
 	SequencerWidget.Reset();
 }
@@ -304,7 +303,7 @@ UMovieScene* FSequencer::GetFocusedMovieScene() const
 	return MovieSceneStack.Top()->GetMovieScene();
 }
 
-void FSequencer::ResetToNewRootMovieScene( UMovieScene& NewRoot, TSharedRef<ISequencerObjectBindingManager> NewObjectBindingManager )
+void FSequencer::ResetToNewAnimation(IMovieSceneAnimation& NewAnimation)
 {
 	DestroySpawnablesForAllMovieScenes();
 
@@ -316,12 +315,13 @@ void FSequencer::ResetToNewRootMovieScene( UMovieScene& NewRoot, TSharedRef<ISeq
 	UnfilterableObjects.Empty();
 	MovieSceneSectionToInstanceMap.Empty();
 
-	NewRoot.SetFlags(RF_Transactional);
+	NewAnimation.GetMovieScene()->SetFlags(RF_Transactional);
 
-	ObjectBindingManager = NewObjectBindingManager;
+	Animation = &NewAnimation;
+	check(Animation != nullptr);
 
 	// Focusing the initial movie scene needs to be done before the first time GetFocusedMovieSceneInstane or GetRootMovieSceneInstance is used
-	RootMovieSceneInstance = MakeShareable(new FMovieSceneInstance(NewRoot));
+	RootMovieSceneInstance = MakeShareable(new FMovieSceneInstance(*Animation->GetMovieScene()));
 	MovieSceneStack.Add(RootMovieSceneInstance.ToSharedRef());
 
 	SequencerWidget->ResetBreadcrumbs();
@@ -452,9 +452,9 @@ void FSequencer::DeleteSelectedKeys()
 
 void FSequencer::SpawnOrDestroyPuppetObjects( TSharedRef<FMovieSceneInstance> MovieSceneInstance )
 {
-	if( ObjectBindingManager->AllowsSpawnableObjects() )
+	if( Animation->AllowsSpawnableObjects() )
 	{
-		ObjectBindingManager->SpawnOrDestroyObjectsForInstance( MovieSceneInstance, false /*bDestroyAll*/ );
+		Animation->SpawnOrDestroyObjects(MovieSceneInstance->GetMovieScene(), false /*DestroyAll*/);
 	}
 }
 
@@ -493,13 +493,13 @@ void FSequencer::OnActorsDropped( const TArray<TWeakObjectPtr<AActor> >& Actors
 			OwnerMovieScene->Modify();
 				
 			const FGuid PossessableGuid = OwnerMovieScene->AddPossessable(Actor->GetActorLabel(), Actor->GetClass());
-				 
+
 			if (IsShotFilteringOn())
 			{
 				AddUnfilterableObject(PossessableGuid);
 			}
-				
-			ObjectBindingManager->BindPossessableObject(PossessableGuid, *Actor);			
+
+			Animation->BindPossessableObject(PossessableGuid, *Actor);
 			bPossessableAdded = true;
 		}
 	}
@@ -616,9 +616,14 @@ void FSequencer::SetPerspectiveViewportPossessionEnabled(bool bEnabled)
 
 FGuid FSequencer::GetHandleToObject( UObject* Object )
 {
+	if (Object == nullptr)
+	{
+		return FGuid();
+	}
+
 	TSharedRef<FMovieSceneInstance> FocusedMovieSceneInstance = GetFocusedMovieSceneInstance();
 	UMovieScene* FocusedMovieScene = FocusedMovieSceneInstance->GetMovieScene();
-	FGuid ObjectGuid = ObjectBindingManager->FindGuidForObject( *FocusedMovieScene, *Object );
+	FGuid ObjectGuid = Animation->FindObjectId(*Object);
 
 	// Check here for spawnable otherwise spawnables get recreated as possessables, which doesn't make sens
 	FMovieSceneSpawnable* Spawnable = FocusedMovieScene->FindSpawnable(ObjectGuid);
@@ -633,10 +638,9 @@ FGuid FSequencer::GetHandleToObject( UObject* Object )
 		// Make sure that the possessable is still valid, if it's not remove the binding so new one 
 		// can be created.  This can happen due to undo.
 		FMovieScenePossessable* Possessable = FocusedMovieScene->FindPossessable(ObjectGuid);
-		
-		if (Possessable == nullptr)
+		if(Possessable == nullptr)
 		{
-			ObjectBindingManager->UnbindPossessableObjects(ObjectGuid);
+			Animation->UnbindPossessableObjects(ObjectGuid);
 			ObjectGuid.Invalidate();
 		}
 	}
@@ -645,7 +649,7 @@ FGuid FSequencer::GetHandleToObject( UObject* Object )
 
 	// If the object guid was not found attempt to add it
 	// Note: Only possessed actors can be added like this
-	if (!ObjectGuid.IsValid() && ObjectBindingManager->CanPossessObject(*Object))
+	if (!ObjectGuid.IsValid() && Animation->CanPossessObject(*Object))
 	{
 		// @todo sequencer: Undo doesn't seem to be working at all
 		const FScopedTransaction Transaction(LOCTEXT("UndoPossessingObject", "Possess Object with MovieScene"));
@@ -657,17 +661,17 @@ FGuid FSequencer::GetHandleToObject( UObject* Object )
 
 			ObjectGuid = FocusedMovieScene->AddPossessable(Object->GetName(), Object->GetClass());
 
-			if (IsShotFilteringOn())
+			if(IsShotFilteringOn())
 			{
 				AddUnfilterableObject(ObjectGuid);
 			}
 
-			ObjectBindingManager->BindPossessableObject(ObjectGuid, *Object);
+			Animation->BindPossessableObject(ObjectGuid, *Object);
 			bPossessableAdded = true;
 		}
 
 		// If we're adding a possessable, generate handles to it's parent objects too.
-		UObject* ParentObject = ObjectBindingManager->GetParentObject(Object);
+		UObject* ParentObject = Animation->GetParentObject(Object);
 		if ( ParentObject != nullptr )
 		{
 			GetHandleToObject(ParentObject);
@@ -696,36 +700,17 @@ void FSequencer::SpawnActorsForMovie( TSharedRef<FMovieSceneInstance> MovieScene
 
 void FSequencer::GetRuntimeObjects( TSharedRef<FMovieSceneInstance> MovieSceneInstance, const FGuid& ObjectHandle, TArray< UObject*>& OutObjects ) const
 {
-	ObjectBindingManager->GetRuntimeObjects( MovieSceneInstance, ObjectHandle, OutObjects );
-	
-	/*if( ObjectSpawner.IsValid() )
-	{	
-		// First, try to find spawnable puppet objects for the specified Guid
-		UObject* FoundPuppetObject = ObjectSpawner->FindPuppetObjectForSpawnableGuid( MovieSceneInstance, ObjectHandle );
-		if( FoundPuppetObject != nullptr )
-		{
-			// Found a puppet object!
-			OutObjects.Reset();
-			OutObjects.Add( FoundPuppetObject );
-		}
-		else
-		{
-			// No puppets were found for spawnables, so now we'll check for possessed actors
-			if( PlayMovieSceneNode != nullptr )
-			{
-				// @todo Sequencer SubMovieScene: Figure out how to instance possessables
-				OutObjects = PlayMovieSceneNode->FindBoundObjects( ObjectHandle );
-			}
-		}
+	if (Animation == nullptr)
+	{
+		return;
 	}
-	else
+
+	UObject* FoundObject = Animation->FindObject(ObjectHandle);
+
+	if (FoundObject != nullptr)
 	{
-		FMovieScenePossessable* Possessable = MovieSceneInstance->GetMovieScene()->FindPossessable( ObjectHandle );
-		if( Possessable && Possessable->GetPossessableObject() )
-		{
-			OutObjects.Add( Possessable->GetPossessableObject() );
-		}
-	}*/
+		OutObjects.Add(FoundObject);
+	}
 }
 
 void FSequencer::UpdateCameraCut(UObject* ObjectToViewThrough, bool bNewCameraCut) const
@@ -761,11 +746,8 @@ void FSequencer::AddOrUpdateMovieSceneInstance( UMovieSceneSection& MovieSceneSe
 
 void FSequencer::RemoveMovieSceneInstance( UMovieSceneSection& MovieSceneSection, TSharedRef<FMovieSceneInstance> InstanceToRemove )
 {
-	if( ObjectBindingManager->AllowsSpawnableObjects() )
-	{
-		ObjectBindingManager->RemoveMovieSceneInstance( InstanceToRemove );
-	}
-	
+	UMovieScene* MovieScene = InstanceToRemove->GetMovieScene();
+	Animation->DestroyAllSpawnedObjects(*MovieScene);
 
 	MovieSceneSectionToInstanceMap.Remove( &MovieSceneSection );
 }
@@ -804,10 +786,8 @@ void FSequencer::UpdateRuntimeInstances()
 
 void FSequencer::DestroySpawnablesForAllMovieScenes()
 {
-	if( ObjectBindingManager->AllowsSpawnableObjects() )
-	{
-		ObjectBindingManager->DestroyAllSpawnedObjects();
-	}
+	UMovieScene* MovieScene = RootMovieSceneInstance->GetMovieScene();
+	Animation->DestroyAllSpawnedObjects(*MovieScene);
 }
 
 FReply FSequencer::OnPlay(bool bTogglePlay)
@@ -1046,7 +1026,7 @@ FGuid FSequencer::AddSpawnableForAssetOrClass( UObject* Object, UObject* Counter
 {
 	FGuid NewSpawnableGuid;
 
-	if (ObjectBindingManager->AllowsSpawnableObjects())
+	if (Animation->AllowsSpawnableObjects())
 	{
 		// Grab the MovieScene that is currently focused.  We'll add our Blueprint as an inner of the
 		// MovieScene asset.
@@ -1194,9 +1174,9 @@ void FSequencer::OnRequestNodeDeleted( TSharedRef<const FSequencerDisplayNode>&
 			// The guid should be associated with a possessable if it wasnt a spawnable
 			bRemoved = OwnerMovieScene->RemovePossessable( BindingToRemove );
 			
-			// @todo Sequencer - undo needs to work here
-			ObjectBindingManager->UnbindPossessableObjects( BindingToRemove );
-			
+			// @todo sequencer: undo needs to work here
+			Animation->UnbindPossessableObjects( BindingToRemove );
+
 			// If this check fails the guid was not associated with a spawnable or possessable so there was an invalid guid being stored on a node
 			check( bRemoved );
 		}
@@ -1510,9 +1490,9 @@ void FSequencer::KeyProperty(FKeyPropertyParams KeyPropertyParams)
 	ObjectChangeListener->KeyProperty(KeyPropertyParams);
 }
 
-TSharedRef<ISequencerObjectBindingManager> FSequencer::GetObjectBindingManager() const
+IMovieSceneAnimation* FSequencer::GetAnimation()
 {
-	return ObjectBindingManager.ToSharedRef();
+	return Animation;
 }
 
 FSequencerSelection& FSequencer::GetSelection()
diff --git a/Engine/Source/Editor/Sequencer/Private/Sequencer.h b/Engine/Source/Editor/Sequencer/Private/Sequencer.h
--- a/Engine/Source/Editor/Sequencer/Private/Sequencer.h
+++ b/Engine/Source/Editor/Sequencer/Private/Sequencer.h
@@ -9,33 +9,10 @@
 
 class UMovieScene;
 class IToolkitHost;
-class ISequencerObjectBindingManager;
+class IMovieSceneAnimation;
 class ISequencerObjectChangeListener;
 
 
-/** Parameters for initializing a Sequencer */
-struct FSequencerInitParams
-{
-	/** The root movie scene being edited */
-	UMovieScene* RootMovieScene;
-	
-	/** The interface for managing runtime object bindings */
-	TSharedPtr<ISequencerObjectBindingManager> ObjectBindingManager;
-	
-	/** The Object change listener for various systems that need to be notified when objects change */
-	TSharedPtr<ISequencerObjectChangeListener> ObjectChangeListener;
-	
-	/** The asset editor created for this (if any) */
-	TSharedPtr<IToolkitHost> ToolkitHost;
-
-	/** View parameters */
-	FSequencerViewParams ViewParams;
-
-	/** Whether or not sequencer should be edited within the level editor */
-	bool bEditWithinLevelEditor;
-};
-
-
 /**
  * Sequencer is the editing tool for MovieScene assets.
  */
@@ -52,11 +29,11 @@ public:
 	/**
 	 * Initializes sequencer
 	 *
-	 * @param	InitParams		Initialization options
-	 * @param	TrackEditorDelegates	Delegates to call to create auto-key handlers for this sequencer
-
+	 * @param InitParams Initialization parameters.
+	 * @param InObjectChangeListener The object change listener to use.
+	 * @param TrackEditorDelegates Delegates to call to create auto-key handlers for this sequencer.
 	 */
-	void InitSequencer( const FSequencerInitParams& InitParams, const TArray<FOnCreateTrackEditor>& TrackEditorDelegates );
+	void InitSequencer(const FSequencerInitParams& InitParams, const TSharedRef<ISequencerObjectChangeListener>& InObjectChangeListener, const TArray<FOnCreateTrackEditor>& TrackEditorDelegates);
 
 	/** Constructor */
 	FSequencer();
@@ -86,7 +63,7 @@ public:
 	virtual TSharedRef<SWidget> GetSequencerWidget() const override { return SequencerWidget.ToSharedRef(); }
 	virtual UMovieScene* GetRootMovieScene() const override;
 	virtual UMovieScene* GetFocusedMovieScene() const override;
-	virtual void ResetToNewRootMovieScene( UMovieScene& NewRoot, TSharedRef<ISequencerObjectBindingManager> NewObjectBindingManager ) override;
+	virtual void ResetToNewAnimation(IMovieSceneAnimation& NewAnimation) override;
 	virtual TSharedRef<FMovieSceneInstance> GetRootMovieSceneInstance() const override;
 	virtual TSharedRef<FMovieSceneInstance> GetFocusedMovieSceneInstance() const override;
 	virtual void FocusSubMovieScene( TSharedRef<FMovieSceneInstance> SubMovieSceneInstance ) override;
@@ -109,7 +86,7 @@ public:
 	virtual void FilterToSelectedShotSections(bool bZoomToShotBounds = true) override;
 	virtual bool CanKeyProperty(FCanKeyPropertyParams CanKeyPropertyParams) const override;
 	virtual void KeyProperty(FKeyPropertyParams KeyPropertyParams) override;
-	virtual TSharedRef<ISequencerObjectBindingManager> GetObjectBindingManager() const override;
+	virtual IMovieSceneAnimation* GetAnimation() override;
 	virtual FSequencerSelection& GetSelection() override;
 	virtual void NotifyMapChanged(class UWorld* NewWorld, EMapChangeType MapChangeType) override;
 
@@ -423,8 +400,8 @@ private:
 	TWeakPtr<FMovieSceneTrackEditor> ShotTrackEditor;
 	TWeakPtr<FMovieSceneTrackEditor> AnimationTrackEditor;
 
-	/** Manager for handling runtime object bindings */
-	TSharedPtr< ISequencerObjectBindingManager > ObjectBindingManager;
+	/** The current animation being edited. */
+	IMovieSceneAnimation* Animation;
 
 	/** Listener for object changes being made while this sequencer is open*/
 	TSharedPtr< class ISequencerObjectChangeListener > ObjectChangeListener;
diff --git a/Engine/Source/Editor/Sequencer/Private/SequencerActorBindingManager.cpp b/Engine/Source/Editor/Sequencer/Private/SequencerActorBindingManager.cpp
deleted file mode 100644
diff --git a/Engine/Source/Editor/Sequencer/Private/SequencerActorBindingManager.h b/Engine/Source/Editor/Sequencer/Private/SequencerActorBindingManager.h
deleted file mode 100644
diff --git a/Engine/Source/Editor/Sequencer/Private/SequencerAssetEditor.cpp b/Engine/Source/Editor/Sequencer/Private/SequencerAssetEditor.cpp
deleted file mode 100644
diff --git a/Engine/Source/Editor/Sequencer/Private/SequencerAssetEditor.h b/Engine/Source/Editor/Sequencer/Private/SequencerAssetEditor.h
deleted file mode 100644
diff --git a/Engine/Source/Editor/Sequencer/Private/SequencerDisplayNode.cpp b/Engine/Source/Editor/Sequencer/Private/SequencerDisplayNode.cpp
--- a/Engine/Source/Editor/Sequencer/Private/SequencerDisplayNode.cpp
+++ b/Engine/Source/Editor/Sequencer/Private/SequencerDisplayNode.cpp
@@ -12,7 +12,6 @@
 #include "MovieSceneTrack.h"
 #include "CommonMovieSceneTools.h"
 #include "IKeyArea.h"
-#include "ISequencerObjectBindingManager.h"
 
 
 #define LOCTEXT_NAMESPACE "SequencerDisplayNode"
@@ -577,7 +576,7 @@ void FTrackNode::FixRowIndices()
 FText FObjectBindingNode::GetDisplayName() const
 {
 	FText DisplayName;
-	return GetSequencer().GetObjectBindingManager()->TryGetObjectBindingDisplayName(GetSequencer().GetFocusedMovieSceneInstance(), ObjectBinding, DisplayName) ?
+	return GetSequencer().GetAnimation()->TryGetObjectDisplayName(ObjectBinding, DisplayName) ?
 		DisplayName : DefaultDisplayName;
 }
 
@@ -697,10 +696,10 @@ void GetKeyablePropertyPaths(UClass* Class, UStruct* PropertySource, TArray<UPro
 TSharedRef<SWidget> FObjectBindingNode::OnGetAddPropertyTrackMenuContent()
 {
 	TArray<TArray<UProperty*>> KeyablePropertyPaths;
-	TArray<UObject*> BoundObjects;
 	FSequencer& Sequencer = GetSequencer();
-	Sequencer.GetObjectBindingManager()->GetRuntimeObjects(Sequencer.GetRootMovieSceneInstance(), ObjectBinding, BoundObjects);
-	for (UObject* BoundObject : BoundObjects)
+	UObject* BoundObject = Sequencer.GetAnimation()->FindObject(ObjectBinding);
+
+	if (BoundObject != nullptr)
 	{
 		TArray<UProperty*> PropertyPath;
 		GetKeyablePropertyPaths(BoundObject->GetClass(), BoundObject->GetClass(), PropertyPath, Sequencer, KeyablePropertyPaths);
@@ -722,7 +721,7 @@ TSharedRef<SWidget> FObjectBindingNode::OnGetAddPropertyTrackMenuContent()
 
 	ISequencerModule& SequencerModule = FModuleManager::GetModuleChecked<ISequencerModule>( "Sequencer" );
 	TSharedRef<FUICommandList> CommandList(new FUICommandList);
-	FMenuBuilder AddTrackMenuBuilder(true, nullptr, SequencerModule.GetMenuExtensibilityManager()->GetAllExtenders(CommandList, BoundObjects));
+	FMenuBuilder AddTrackMenuBuilder(true, nullptr, SequencerModule.GetMenuExtensibilityManager()->GetAllExtenders(CommandList, TArrayBuilder<UObject*>().Add(BoundObject)));
 
 	AddTrackMenuBuilder.BeginSection( SequencerMenuExtensionPoints::AddTrackMenu_PropertiesSection, LOCTEXT("PropertiesMenuHeader" , "Properties"));
 	{
@@ -745,12 +744,10 @@ TSharedRef<SWidget> FObjectBindingNode::OnGetAddPropertyTrackMenuContent()
 void FObjectBindingNode::AddTrackForProperty(TArray<UProperty*> PropertyPath)
 {
 	FSequencer& Sequencer = GetSequencer();
-
-	TArray<UObject*> BoundObjects;
-	Sequencer.GetObjectBindingManager()->GetRuntimeObjects(Sequencer.GetFocusedMovieSceneInstance(), ObjectBinding, BoundObjects);
+	UObject* BoundObject = Sequencer.GetAnimation()->FindObject(ObjectBinding);
 
 	TArray<UObject*> KeyableBoundObjects;
-	for (UObject* BoundObject : BoundObjects)
+	if (BoundObject != nullptr)
 	{
 		if (Sequencer.CanKeyProperty(FCanKeyPropertyParams(BoundObject->GetClass(), PropertyPath)))
 		{
diff --git a/Engine/Source/Editor/Sequencer/Private/SequencerModule.cpp b/Engine/Source/Editor/Sequencer/Private/SequencerModule.cpp
--- a/Engine/Source/Editor/Sequencer/Private/SequencerModule.cpp
+++ b/Engine/Source/Editor/Sequencer/Private/SequencerModule.cpp
@@ -5,7 +5,6 @@
 #include "Sequencer.h"
 #include "Toolkits/ToolkitManager.h"
 #include "SequencerCommands.h"
-#include "SequencerAssetEditor.h"
 #include "SequencerEdMode.h"
 #include "SequencerObjectChangeListener.h"
 
@@ -28,33 +27,16 @@ public:
 
 	// ISequencerModule interface
 
-	virtual TSharedPtr<ISequencer> CreateSequencer( UMovieScene* RootMovieScene, const FSequencerViewParams& InViewParams, TSharedRef<ISequencerObjectBindingManager> ObjectBindingManager ) override
+	virtual TSharedRef<ISequencer> CreateSequencer(const FSequencerInitParams& InitParams) override
 	{
-		TSharedRef< FSequencer > Sequencer = MakeShareable(new FSequencer);
-		
-		FSequencerInitParams SequencerInitParams;
-		{
-			SequencerInitParams.ViewParams = InViewParams;
-			SequencerInitParams.ObjectChangeListener = MakeShareable( new FSequencerObjectChangeListener( Sequencer, false ) );
-			SequencerInitParams.ObjectBindingManager = ObjectBindingManager;
-			SequencerInitParams.RootMovieScene = RootMovieScene;
-			SequencerInitParams.bEditWithinLevelEditor = false;
-			SequencerInitParams.ToolkitHost = nullptr;
-		}
-		
-		Sequencer->InitSequencer(SequencerInitParams, TrackEditorDelegates);
-		
+		TSharedRef<FSequencer> Sequencer = MakeShareable(new FSequencer);
+		TSharedRef<ISequencerObjectChangeListener> ObjectChangeListener = MakeShareable(new FSequencerObjectChangeListener(Sequencer, InitParams.bEditWithinLevelEditor));
+
+		Sequencer->InitSequencer(InitParams, ObjectChangeListener, TrackEditorDelegates);
+
 		return Sequencer;
 	}
 	
-	virtual TSharedPtr<ISequencer> CreateSequencerAssetEditor( const EToolkitMode::Type Mode, const FSequencerViewParams& InViewParams, const TSharedPtr< class IToolkitHost >& InitToolkitHost, UMovieScene* InRootMovieScene, bool bEditWithinLevelEditor ) override
-	{
-		TSharedRef<FSequencerAssetEditor> SequencerAssetEditor = MakeShareable(new FSequencerAssetEditor);
-
-		SequencerAssetEditor->InitSequencerAssetEditor( Mode, InViewParams, InitToolkitHost, InRootMovieScene, TrackEditorDelegates, bEditWithinLevelEditor );
-		return SequencerAssetEditor->GetSequencerInterface();
-	}
-
 	virtual FDelegateHandle RegisterTrackEditor_Handle( FOnCreateTrackEditor InOnCreateTrackEditor ) override
 	{
 		TrackEditorDelegates.Add( InOnCreateTrackEditor );
diff --git a/Engine/Source/Editor/Sequencer/Private/SequencerMovieSceneObjectManager.cpp b/Engine/Source/Editor/Sequencer/Private/SequencerMovieSceneObjectManager.cpp
deleted file mode 100644
diff --git a/Engine/Source/Editor/Sequencer/Private/SequencerMovieSceneObjectManager.h b/Engine/Source/Editor/Sequencer/Private/SequencerMovieSceneObjectManager.h
deleted file mode 100644
diff --git a/Engine/Source/Editor/Sequencer/Private/SequencerNodeTree.cpp b/Engine/Source/Editor/Sequencer/Private/SequencerNodeTree.cpp
--- a/Engine/Source/Editor/Sequencer/Private/SequencerNodeTree.cpp
+++ b/Engine/Source/Editor/Sequencer/Private/SequencerNodeTree.cpp
@@ -1,6 +1,7 @@
 // Copyright 1998-2015 Epic Games, Inc. All Rights Reserved.
 
 #include "SequencerPrivatePCH.h"
+#include "MovieSceneAnimation.h"
 #include "MovieSceneSection.h"
 #include "SequencerNodeTree.h"
 #include "Sequencer.h"
@@ -10,7 +11,7 @@
 #include "MovieSceneTrackEditor.h"
 #include "SectionLayoutBuilder.h"
 #include "ISequencerSection.h"
-#include "ISequencerObjectBindingManager.h"
+
 
 void FSequencerNodeTree::Empty()
 {
@@ -200,14 +201,14 @@ TSharedRef<FObjectBindingNode> FSequencerNodeTree::AddObjectBinding(const FStrin
 		// Try to get the parent object node if there is one.
 		TSharedPtr<FObjectBindingNode> ParentNode;
 		TArray<UObject*> RuntimeObjects;
-		TSharedRef<ISequencerObjectBindingManager> BindingManager(Sequencer.GetObjectBindingManager());
-		BindingManager->GetRuntimeObjects(Sequencer.GetFocusedMovieSceneInstance(), ObjectBinding, RuntimeObjects);
-		if ( RuntimeObjects.Num() == 1)
+		IMovieSceneAnimation* Animation = Sequencer.GetAnimation();
+		UObject* RuntimeObject = Animation->FindObject(ObjectBinding);
+		if ( RuntimeObject != nullptr)
 		{
-			UObject* ParentObject = BindingManager->GetParentObject(RuntimeObjects[0]);
+			UObject* ParentObject = Animation->GetParentObject(RuntimeObject);
 			if (ParentObject != nullptr)
 			{
-				FGuid ParentBinding = BindingManager->FindGuidForObject(*Sequencer.GetFocusedMovieScene(), *ParentObject);
+				FGuid ParentBinding = Animation->FindObjectId(*ParentObject);
 				TSharedPtr<FObjectBindingNode>* FoundParentNode = ObjectBindingMap.Find( ParentBinding );
 				if ( FoundParentNode != nullptr )
 				{
diff --git a/Engine/Source/Editor/Sequencer/Private/SequencerPossessedObject.cpp b/Engine/Source/Editor/Sequencer/Private/SequencerPossessedObject.cpp
deleted file mode 100644
diff --git a/Engine/Source/Editor/Sequencer/Private/SequencerPossessedObject.h b/Engine/Source/Editor/Sequencer/Private/SequencerPossessedObject.h
deleted file mode 100644
diff --git a/Engine/Source/Editor/Sequencer/Public/ISequencer.h b/Engine/Source/Editor/Sequencer/Public/ISequencer.h
--- a/Engine/Source/Editor/Sequencer/Public/ISequencer.h
+++ b/Engine/Source/Editor/Sequencer/Public/ISequencer.h
@@ -9,12 +9,13 @@
 
 
 class FSequencerSelection;
-class ISequencerObjectBindingManager;
+class IMovieSceneAnimation;
 class UAnimSequence;
 class UMovieScene;
 class UMovieSceneSection;
 enum class EMapChangeType : uint8;
 
+
 /**
  * Interface for sequencers.
  */
@@ -41,8 +42,8 @@ public:
 	 */
 	virtual TSharedRef<FMovieSceneInstance> GetFocusedMovieSceneInstance() const = 0;
 
-	/** Resets sequencer with a new RootMovieScene */
-	virtual void ResetToNewRootMovieScene(UMovieScene& NewRoot, TSharedRef<ISequencerObjectBindingManager> NewObjectBindingManager) = 0;
+	/** Resets sequencer with a new animation */
+	virtual void ResetToNewAnimation(IMovieSceneAnimation& NewAnimation) = 0;
 
 	/**
 	 * Focuses a sub-movie scene (MovieScene within a MovieScene) in the sequencer.
@@ -155,7 +156,12 @@ public:
 
 	virtual void UpdateRuntimeInstances() = 0;
 
-	virtual TSharedRef<ISequencerObjectBindingManager> GetObjectBindingManager() const = 0;
+	/**
+	 * Get the animation that is being played by this sequencer.
+	 *
+	 * @return The current animation, or nullptr if no animation is playing.
+	 */
+	virtual IMovieSceneAnimation* GetAnimation() = 0;
 
 	virtual FSequencerSelection& GetSelection() = 0;
 
diff --git a/Engine/Source/Editor/Sequencer/Public/ISequencerAssetEditor.h b/Engine/Source/Editor/Sequencer/Public/ISequencerAssetEditor.h
deleted file mode 100644
diff --git a/Engine/Source/Editor/Sequencer/Public/ISequencerModule.h b/Engine/Source/Editor/Sequencer/Public/ISequencerModule.h
--- a/Engine/Source/Editor/Sequencer/Public/ISequencerModule.h
+++ b/Engine/Source/Editor/Sequencer/Public/ISequencerModule.h
@@ -9,6 +9,8 @@
 
 
 class ISequencerObjectBindingManager;
+class IToolkitHost;
+class UActorAnimation;
 
 
 namespace SequencerMenuExtensionPoints
@@ -25,7 +27,7 @@ DECLARE_DELEGATE_RetVal_OneParam(TSharedRef<SWidget>, FOnGetAddMenuContent, TSha
 
 
 /**
- * Sequencer initialization parameters.
+ * Sequencer view parameters.
  */
 struct FSequencerViewParams
 {
@@ -48,6 +50,25 @@ struct FSequencerViewParams
 };
 
 
+/**
+ * Sequencer initialization parameters.
+ */
+struct FSequencerInitParams
+{
+	/** The animation being edited. */
+	IMovieSceneAnimation* Animation;
+
+	/** The asset editor created for this (if any) */
+	TSharedPtr<IToolkitHost> ToolkitHost;
+
+	/** View parameters */
+	FSequencerViewParams ViewParams;
+
+	/** Whether or not sequencer should be edited within the level editor */
+	bool bEditWithinLevelEditor;
+};
+
+
 /**
  * Interface for the Sequencer module.
  */
@@ -57,24 +78,12 @@ class ISequencerModule
 public:
 
 	/**
-	 * Creates a new instance of a standalone sequencer that can be added to other UIs.
-	 *
-	 * @param InRootMovieScene The movie scene to edit.
-	 * @param ViewParams Parameters for how to view sequencer UI.
-	 * @return Interface to the new editor.
-	 */
-	virtual TSharedPtr<ISequencer> CreateSequencer(UMovieScene* InRootMovieScene, const FSequencerViewParams& InViewParams, TSharedRef<ISequencerObjectBindingManager> ObjectBindingManager) = 0;
-
-	/**
-	 * Creates a new instance of a Sequencer, the editor for MovieScene assets in an asset editor
+	 * Create a new instance of a standalone sequencer that can be added to other UIs.
 	 *
-	 * @param Mode Mode that this editor should operate in.
-	 * @param ViewParams Parameters for how to view sequencer UI.
-	 * @param InitToolkitHost When Mode is WorldCentric, this is the level editor instance to spawn this editor within.
-	 * @param ObjectToEdit The object to start editing.
-	 * @return Interface to the new editor.
+	 * @param InitParams Initialization parameters.
+	 * @return The new sequencer object.
 	 */
-	virtual TSharedPtr<ISequencer> CreateSequencerAssetEditor(const EToolkitMode::Type Mode, const FSequencerViewParams& InViewParams, const TSharedPtr< class IToolkitHost >& InitToolkitHost, UMovieScene* InRootMovieScene, bool bEditWithinLevelEditor) = 0;
+	virtual TSharedRef<ISequencer> CreateSequencer(const FSequencerInitParams& InitParams) = 0;
 
 	/** 
 	 * Registers a delegate that will create an editor for a track in each sequencer.
diff --git a/Engine/Source/Editor/Sequencer/Public/ISequencerObjectBindingManager.h b/Engine/Source/Editor/Sequencer/Public/ISequencerObjectBindingManager.h
deleted file mode 100644
diff --git a/Engine/Source/Editor/UMGEditor/Private/Animation/UMGMovieSceneObjectManager.cpp b/Engine/Source/Editor/UMGEditor/Private/Animation/UMGMovieSceneObjectManager.cpp
deleted file mode 100644
diff --git a/Engine/Source/Editor/UMGEditor/Private/Animation/UMGMovieSceneObjectManager.h b/Engine/Source/Editor/UMGEditor/Private/Animation/UMGMovieSceneObjectManager.h
deleted file mode 100644
diff --git a/Engine/Source/Editor/UMGEditor/Private/Animation/UMGSequencerObjectBindingManager.cpp b/Engine/Source/Editor/UMGEditor/Private/Animation/UMGSequencerObjectBindingManager.cpp
deleted file mode 100644
diff --git a/Engine/Source/Editor/UMGEditor/Private/Animation/UMGSequencerObjectBindingManager.h b/Engine/Source/Editor/UMGEditor/Private/Animation/UMGSequencerObjectBindingManager.h
deleted file mode 100644
diff --git a/Engine/Source/Editor/UMGEditor/Private/WidgetBlueprintEditor.cpp b/Engine/Source/Editor/UMGEditor/Private/WidgetBlueprintEditor.cpp
--- a/Engine/Source/Editor/UMGEditor/Private/WidgetBlueprintEditor.cpp
+++ b/Engine/Source/Editor/UMGEditor/Private/WidgetBlueprintEditor.cpp
@@ -6,7 +6,6 @@
 #include "WidgetBlueprintEditor.h"
 #include "MovieScene.h"
 #include "Editor/Sequencer/Public/ISequencerModule.h"
-#include "Animation/UMGSequencerObjectBindingManager.h"
 #include "ObjectEditorUtils.h"
 
 #include "PropertyCustomizationHelpers.h"
@@ -53,8 +52,6 @@ FWidgetBlueprintEditor::~FWidgetBlueprintEditor()
 	
 	Sequencer.Reset();
 
-	SequencerObjectBindingManager.Reset();
-
 	// Un-Register sequencer menu extenders.
 	ISequencerModule& SequencerModule = FModuleManager::Get().LoadModuleChecked<ISequencerModule>("Sequencer");
 	SequencerModule.GetMenuExtensibilityManager()->GetExtenderDelegates().RemoveAll([this]( const FAssetEditorExtender& Extender )
@@ -86,9 +83,9 @@ void FWidgetBlueprintEditor::InitWidgetBlueprintEditor(const EToolkitMode::Type
 	UpdatePreview(GetWidgetBlueprintObj(), true);
 
 	// If the user has close the sequencer tab, this will not be initialized.
-	if ( SequencerObjectBindingManager.IsValid() )
+	if (CurrentAnimation.IsValid())
 	{
-		SequencerObjectBindingManager->InitPreviewObjects();
+		CurrentAnimation->Initialize(PreviewWidgetPtr.Get());
 	}
 
 	DesignerCommandList = MakeShareable(new FUICommandList);
@@ -312,13 +309,7 @@ void FWidgetBlueprintEditor::CutSelectedWidgets()
 
 const UWidgetAnimation* FWidgetBlueprintEditor::RefreshCurrentAnimation()
 {
-	if( !SequencerObjectBindingManager->HasValidWidgetAnimation() )
-	{
-		ChangeViewedAnimation(*UWidgetAnimation::GetNullAnimation());
-		return nullptr;
-	}
-
-	return SequencerObjectBindingManager->GetWidgetAnimation();
+	return CurrentAnimation.Get();
 }
 
 bool FWidgetBlueprintEditor::CanPasteWidgets()
@@ -564,9 +555,15 @@ TSharedPtr<ISequencer>& FWidgetBlueprintEditor::GetSequencer()
 			ViewParams.OnGetAddMenuContent = FOnGetAddMenuContent::CreateSP(this, &FWidgetBlueprintEditor::OnGetAnimationAddMenuContent);
 		}
 
-		SequencerObjectBindingManager = MakeShareable(new FUMGSequencerObjectBindingManager(*this, *UWidgetAnimation::GetNullAnimation()));
+		FSequencerInitParams SequencerInitParams;
+		{
+			SequencerInitParams.ViewParams = ViewParams;
+			SequencerInitParams.Animation = UWidgetAnimation::GetNullAnimation();
+			SequencerInitParams.bEditWithinLevelEditor = false;
+			SequencerInitParams.ToolkitHost = nullptr;
+		};
 
-		Sequencer = FModuleManager::LoadModuleChecked< ISequencerModule >("Sequencer").CreateSequencer(UWidgetAnimation::GetNullAnimation()->MovieScene, ViewParams, SequencerObjectBindingManager.ToSharedRef());
+		Sequencer = FModuleManager::LoadModuleChecked<ISequencerModule>("Sequencer").CreateSequencer(SequencerInitParams);
 		ChangeViewedAnimation(*UWidgetAnimation::GetNullAnimation());
 	}
 
@@ -575,12 +572,15 @@ TSharedPtr<ISequencer>& FWidgetBlueprintEditor::GetSequencer()
 
 void FWidgetBlueprintEditor::ChangeViewedAnimation( UWidgetAnimation& InAnimationToView )
 {
-	TSharedRef<FUMGSequencerObjectBindingManager> NewObjectBindingManager = MakeShareable(new FUMGSequencerObjectBindingManager(*this, InAnimationToView));
-	Sequencer->ResetToNewRootMovieScene(*InAnimationToView.MovieScene, NewObjectBindingManager);
+	if (CurrentAnimation.IsValid())
+	{
+		CurrentAnimation->Initialize(nullptr);
+	}
+
+	CurrentAnimation = &InAnimationToView;
+	CurrentAnimation->Initialize(PreviewWidgetPtr.Get());
 
-	check(SequencerObjectBindingManager.IsUnique());
-	SequencerObjectBindingManager = NewObjectBindingManager;
-	SequencerObjectBindingManager->InitPreviewObjects();
+	Sequencer->ResetToNewAnimation(InAnimationToView);
 
 	TSharedPtr<SOverlay> SequencerOverlayPin = SequencerOverlay.Pin();
 	if (SequencerOverlayPin.IsValid())
@@ -662,7 +662,13 @@ void FWidgetBlueprintEditor::UpdatePreview(UBlueprint* InBlueprint, bool bInForc
 		PreviewWidgetPtr = PreviewActor;
 	}
 
+	if (CurrentAnimation.IsValid())
+	{
+		CurrentAnimation->Initialize(PreviewActor);
+	}
+
 	OnWidgetPreviewUpdated.Broadcast();
+	Sequencer->UpdateRuntimeInstances();
 }
 
 FGraphAppearanceInfo FWidgetBlueprintEditor::GetGraphAppearance(UEdGraph* InGraph) const
@@ -789,13 +795,16 @@ TSharedRef<SWidget> FWidgetBlueprintEditor::OnGetAnimationAddMenuContent(TShared
 
 	FMenuBuilder AddMenuBuilder(true, nullptr);
 
-	for (FObjectAndDisplayName& BindableObject : BindableObjects)
+	if (CurrentAnimation.IsValid())
 	{
-		FGuid BoundObjectGuid = SequencerObjectBindingManager->FindGuidForObject(*InSequencer->GetFocusedMovieScene(), *BindableObject.Object);
-		if (!BoundObjectGuid.IsValid())
+		for (FObjectAndDisplayName& BindableObject : BindableObjects)
 		{
-			FUIAction AddMenuAction(FExecuteAction::CreateSP(this, &FWidgetBlueprintEditor::AddObjectToAnimation, BindableObject.Object));
-			AddMenuBuilder.AddMenuEntry(BindableObject.DisplayName, FText(), FSlateIcon(), AddMenuAction);
+			FGuid BoundObjectGuid = CurrentAnimation->FindObjectId(*BindableObject.Object);
+			if (BoundObjectGuid.IsValid() == false)
+			{
+				FUIAction AddMenuAction(FExecuteAction::CreateSP(this, &FWidgetBlueprintEditor::AddObjectToAnimation, BindableObject.Object));
+				AddMenuBuilder.AddMenuEntry(BindableObject.DisplayName, FText(), FSlateIcon(), AddMenuAction);
+			}
 		}
 	}
 
diff --git a/Engine/Source/Editor/UMGEditor/Private/WidgetBlueprintEditor.h b/Engine/Source/Editor/UMGEditor/Private/WidgetBlueprintEditor.h
--- a/Engine/Source/Editor/UMGEditor/Private/WidgetBlueprintEditor.h
+++ b/Engine/Source/Editor/UMGEditor/Private/WidgetBlueprintEditor.h
@@ -203,9 +203,6 @@ private:
 	/** A text block which is displayed in the overlay when no animation is selected. */
 	TWeakPtr<STextBlock> NoAnimationTextBlock;
 
-	/** Manager for handling bindings to sequence animations */
-	TSharedPtr<class FUMGSequencerObjectBindingManager> SequencerObjectBindingManager;
-
 	/** The Blueprint associated with the current preview */
 	UWidgetBlueprint* PreviewBlueprint;
 
@@ -244,5 +241,8 @@ private:
 
 	TArray< TFunction<void()> > QueuedDesignerActions;
 
+	/** The currently viewed animation, if any. */
+	TWeakObjectPtr<UWidgetAnimation> CurrentAnimation;
+
 	FDelegateHandle SequencerExtenderHandle;
 };
diff --git a/Engine/Source/Runtime/ActorAnimation/ActorAnimation.Build.cs b/Engine/Source/Runtime/ActorAnimation/ActorAnimation.Build.cs
new file mode 100644
--- /dev/null
+++ b/Engine/Source/Runtime/ActorAnimation/ActorAnimation.Build.cs
@@ -0,0 +1,20 @@
+// Copyright 1998-2015 Epic Games, Inc. All Rights Reserved.
+
+using UnrealBuildTool;
+
+public class ActorAnimation : ModuleRules
+{
+	public ActorAnimation(TargetInfo Target)
+	{
+		PrivateIncludePaths.Add("Runtime/ActorAnimation/Private");
+
+		PublicDependencyModuleNames.AddRange(
+            new string[] {
+                "Core",
+                "CoreUObject",
+                "Engine",
+				"MovieScene",
+            }
+        );
+	}
+}
diff --git a/Engine/Source/Runtime/ActorAnimation/Private/ActorAnimation.cpp b/Engine/Source/Runtime/ActorAnimation/Private/ActorAnimation.cpp
new file mode 100644
--- /dev/null
+++ b/Engine/Source/Runtime/ActorAnimation/Private/ActorAnimation.cpp
@@ -0,0 +1,332 @@
+// Copyright 1998-2015 Epic Games, Inc. All Rights Reserved.
+
+#include "ActorAnimationPrivatePCH.h"
+#include "ActorAnimation.h"
+#include "ActorAnimationObject.h"
+#include "MovieScene.h"
+
+
+/* UActorAnimation structors
+ *****************************************************************************/
+
+UActorAnimation::UActorAnimation(const FObjectInitializer& ObjectInitializer)
+	: Super(ObjectInitializer)
+	, MovieScene(nullptr)
+{ }
+
+
+/* UActorAnimation interface
+ *****************************************************************************/
+
+void UActorAnimation::Initialize()
+{
+	// @todo sequencer: gmp: fix me
+	MovieScene = NewObject<UMovieScene>(this);
+}
+
+
+/* IMovieSceneAnimation interface
+ *****************************************************************************/
+
+bool UActorAnimation::AllowsSpawnableObjects() const
+{
+	return true;
+}
+
+
+void UActorAnimation::BindPossessableObject(const FGuid& ObjectId, UObject& PossessedObject)
+{
+	FString ComponentName;
+	UActorComponent* Component = Cast<UActorComponent>(&PossessedObject);
+
+	if (Component != nullptr)
+	{
+		// @todo sequencer: gmp: need support for UStruct keys in TMap
+		PossessedObjects.Add(ObjectId.ToString(), FActorAnimationObject(PossessedObject.GetOuter(), ComponentName));
+		//PossessedObjects.Add(ObjectId, FSequencerPossessedObject(PossessedObject.GetOuter(), ComponentName));
+	}
+	else
+	{
+		// @todo sequencer: gmp: need support for UStruct keys in TMap
+		PossessedObjects.Add(ObjectId.ToString(), FActorAnimationObject(&PossessedObject));
+		//PossessedObjects.Add(ObjectId, FSequencerPossessedObject(Actor));
+	}
+}
+
+
+bool UActorAnimation::CanPossessObject(UObject& Object) const
+{
+	return Object.IsA<AActor>();
+}
+
+
+void UActorAnimation::DestroyAllSpawnedObjects(UMovieScene& MovieScene)
+{
+	const bool DestroyAll = true;
+	SpawnOrDestroyObjects(&MovieScene, DestroyAll);
+	SpawnedActors.Empty();
+}
+
+
+UObject* UActorAnimation::FindObject(const FGuid& ObjectId) const
+{
+	// search bindings (possessables)
+	// @todo sequencer: gmp: need support for UStruct keys in TMap
+	const FActorAnimationObject* PossessedObject = PossessedObjects.Find(ObjectId.ToString());
+	//const FSequencerPossessedObject* PossessedObject = PossessedObjects.Find(ObjectId);
+
+	if (PossessedObject != nullptr)
+	{
+		return PossessedObject->GetObject();
+	}
+
+	// search proxies (spawnables)
+	TWeakObjectPtr<AActor> Actor = SpawnedActors.FindRef(ObjectId);
+
+	if (Actor.IsValid())
+	{
+		return Actor.Get();
+	}
+
+	// not found
+	return nullptr;
+}
+
+
+FGuid UActorAnimation::FindObjectId(UObject& Object) const
+{
+	// search possessed objects
+	for (auto PossessedObjectPair : PossessedObjects)
+	{
+		const FActorAnimationObject& PossessedObject = PossessedObjectPair.Value;
+
+		if (PossessedObject.GetObject() == &Object)
+		{
+			// @todo sequencer: gmp: need support for UStruct keys in TMap
+			FGuid Result; FGuid::Parse(PossessedObjectPair.Key, Result); return Result;
+			//return PossessedObjectsPair.Key;
+		}
+	}
+
+	// search spawned objects
+	for (auto SpawnedActorPair : SpawnedActors)
+	{
+		if (SpawnedActorPair.Value.Get() == &Object)
+		{
+			return SpawnedActorPair.Key;
+		}
+	}
+
+	// not found
+	return FGuid();
+}
+
+
+UMovieScene* UActorAnimation::GetMovieScene() const
+{
+	return MovieScene;
+}
+
+
+UObject* UActorAnimation::GetParentObject(UObject* Object) const
+{
+	UActorComponent* Component = Cast<UActorComponent>(Object);
+
+	if (Component != nullptr)
+	{
+		return Component->GetOwner();
+	}
+
+	return nullptr;
+}
+
+
+void UActorAnimation::SpawnOrDestroyObjects(UMovieScene* MovieScene, bool DestroyAll)
+{
+	bool bAnyLevelActorsChanged = false;
+
+	// remove any proxy actors that we no longer need
+	if (MovieScene != nullptr)
+	{
+		for (auto SpawnedActorPair : SpawnedActors)
+		{
+			bool ShouldDestroyActor = true;
+
+			if (!DestroyAll)
+			{
+				// figure out if we still need this proxy actor
+				for (auto SpawnableIndex = 0; SpawnableIndex < MovieScene->GetSpawnableCount(); ++SpawnableIndex)
+				{
+					auto& Spawnable = MovieScene->GetSpawnable(SpawnableIndex);
+
+					if (Spawnable.GetGuid() == SpawnedActorPair.Key)
+					{
+						ShouldDestroyActor = false;
+						break;
+					}
+				}
+			}
+
+			if (!ShouldDestroyActor)
+			{
+				// Actor is no longer valid (probably the world was destroyed)
+				continue;
+			}
+
+			AActor* Actor = SpawnedActorPair.Value.Get();
+				
+			if (Actor == nullptr)
+			{
+				continue;
+			}
+
+			// destroy Actor
+			UWorld* World = Actor->GetWorld();
+
+			if (ensure(World != nullptr))
+			{
+				if (!bAnyLevelActorsChanged)
+				{
+					DeselectAllActors();
+				}
+
+				// Actor should never be selected in the Editor at this point (we took care of that up above)
+				ensure(!Actor->IsSelected());
+
+				const bool bNetForce = false;
+				const bool bShouldModifyLevel = false;
+				const bool bWasDestroyed = World->DestroyActor(Actor, bNetForce, bShouldModifyLevel);
+
+				if (bWasDestroyed)
+				{
+					bAnyLevelActorsChanged = true;
+					SpawnedActors.Remove(SpawnedActorPair.Key);
+				}
+				else
+				{
+					// @todo sequencer: At least one proxy actor couldn't be cleaned up!
+				}
+			}
+		}
+	}
+
+	if (!DestroyAll && (MovieScene != nullptr))
+	{
+		for (auto SpawnableIndex = 0; SpawnableIndex < MovieScene->GetSpawnableCount(); ++SpawnableIndex)
+		{
+			FMovieSceneSpawnable& Spawnable = MovieScene->GetSpawnable( SpawnableIndex );
+
+			// Do we already have a proxy for this spawnable?
+			if (SpawnedActors.Contains(Spawnable.GetGuid()))
+			{
+				continue;
+			}
+
+			UClass* GeneratedClass = Spawnable.GetClass();
+
+			if ((GeneratedClass == nullptr) || !GeneratedClass->IsChildOf(AActor::StaticClass()))
+			{
+				continue;
+			}
+
+			AActor* ActorCDO = CastChecked<AActor>(GeneratedClass->ClassDefaultObject);
+
+			const FVector SpawnLocation = ActorCDO->GetRootComponent()->RelativeLocation;
+			const FRotator SpawnRotation = ActorCDO->GetRootComponent()->RelativeRotation;
+
+			// @todo sequencer: We should probably spawn these in a specific sub-level!
+			// World->CurrentLevel = ???;
+
+			const FName ActorName = NAME_None;
+
+			// Override the object flags so that RF_Transactional is not set.  Puppet actors are never transactional
+			// @todo sequencer: These actors need to avoid any transaction history.  However, RF_Transactional can currently be set on objects on the fly!
+			const EObjectFlags ObjectFlags = RF_Transient;		// NOTE: We are omitting RF_Transactional intentionally
+
+			// @todo sequencer livecapture: Consider using SetPlayInEditorWorld() and RestoreEditorWorld() here instead
+						
+			// @todo sequencer actors: We need to make sure puppet objects aren't copied into PIE/SIE sessions!  They should be omitted from that duplication!
+
+			// Spawn the puppet actor
+			FActorSpawnParameters SpawnInfo;
+			{
+				SpawnInfo.Name = ActorName;
+				SpawnInfo.ObjectFlags = ObjectFlags;
+			}
+
+			AActor* NewActor = GWorld->SpawnActor(GeneratedClass, &SpawnLocation, &SpawnRotation, SpawnInfo);
+				
+			if (NewActor != nullptr)
+			{
+				// @todo sequencer: We're naming the actor based off of the spawnable's name.  Is that really what we want?
+				// @todo sequencer: gmp: fix me
+				//FActorLabelUtilities::SetActorLabelUnique(NewActor, Spawnable.GetName());
+				SpawnedActors.Add(Spawnable.GetGuid(), NewActor);
+			}
+			else
+			{
+				// Actor failed to spawn
+				// @todo sequencer: What should we do when this happens to one or more actors?
+			}
+		}
+	}
+}
+
+
+#if WITH_EDITOR
+bool UActorAnimation::TryGetObjectDisplayName(const FGuid& ObjectId, FText& OutDisplayName) const
+{
+	UObject* Object = FindObject(ObjectId);
+
+	if (Object == nullptr)
+	{
+		return false;
+	}
+
+	AActor* Actor = Cast<AActor>(Object);
+
+	if (Actor == nullptr)
+	{
+		return false;
+	}
+
+	OutDisplayName = FText::FromString(Actor->GetActorLabel());
+
+	return true;
+}
+#endif
+
+
+void UActorAnimation::UnbindPossessableObjects(const FGuid& ObjectId)
+{
+	// @todo sequencer: gmp: need support for UStruct keys in TMap
+	PossessedObjects.Remove(ObjectId.ToString());
+	//PossessedObjects.Remove(ObjectId);
+}
+
+
+/* UActorAnimation implementation
+ *****************************************************************************/
+
+void UActorAnimation::DeselectAllActors()
+{
+	// @todo sequencer: gmp: fix me
+	/*
+	USelection* Selection = GEditor->GetSelectedActors();
+	{
+		Selection->BeginBatchSelectOperation();
+		Selection->MarkBatchDirty();
+
+		for (auto SpawnedActorPair : SpawnedActors)
+		{
+			AActor* SpawnedActor = SpawnedActorPair.Value.Get();
+
+			if (SpawnedActor != nullptr)
+			{
+				GEditor->GetSelectedActors()->Deselect(SpawnedActor);
+			}
+		}
+
+		Selection->EndBatchSelectOperation();
+	}*/
+}
diff --git a/Engine/Source/Runtime/ActorAnimation/Private/ActorAnimationModule.cpp b/Engine/Source/Runtime/ActorAnimation/Private/ActorAnimationModule.cpp
new file mode 100644
--- /dev/null
+++ b/Engine/Source/Runtime/ActorAnimation/Private/ActorAnimationModule.cpp
@@ -0,0 +1,22 @@
+// Copyright 1998-2015 Epic Games, Inc. All Rights Reserved.
+
+#include "ActorAnimationPrivatePCH.h"
+#include "ModuleInterface.h"
+
+
+/**
+ * Implements the ActorAnimation module.
+ */
+class FActorAnimationModule
+	: public IModuleInterface
+{
+public:
+
+	// IModuleInterface interface
+
+	virtual void StartupModule() override { }
+	virtual void ShutdownModule() override { }
+};
+
+
+IMPLEMENT_MODULE(FActorAnimationModule, ActorAnimation);
diff --git a/Engine/Source/Runtime/ActorAnimation/Private/ActorAnimationObject.cpp b/Engine/Source/Runtime/ActorAnimation/Private/ActorAnimationObject.cpp
new file mode 100644
--- /dev/null
+++ b/Engine/Source/Runtime/ActorAnimation/Private/ActorAnimationObject.cpp
@@ -0,0 +1,53 @@
+// Copyright 1998-2015 Epic Games, Inc. All Rights Reserved.
+
+#include "ActorAnimationPrivatePCH.h"
+#include "ActorAnimationObject.h"
+
+
+/* FSequencerPossessedObject interface
+ *****************************************************************************/
+
+UObject* FActorAnimationObject::GetObject() const
+{
+	// get component-less object
+	if (ComponentName.IsEmpty())
+	{
+		return ObjectOrOwner.Get();
+	}
+
+	// get cached component
+	if ((CachedComponent != nullptr) && (CachedComponent->GetName() == ComponentName))
+	{
+		return CachedComponent;
+	}
+
+	CachedComponent = nullptr;
+
+	// find & cache component
+	UObject* Object = ObjectOrOwner.Get();
+
+	if (Object == nullptr)
+	{
+		return nullptr;
+	}
+
+	AActor* Owner = Cast<AActor>(Object);
+
+	if (Owner == nullptr)
+	{
+		return Object;
+	}
+
+	for (UActorComponent* ActorComponent : Owner->GetComponents())
+	{
+		if (ActorComponent->GetName() == ComponentName)
+		{
+			CachedComponent = ActorComponent;
+
+			return ActorComponent;
+		}
+	}
+
+	// component not found
+	return nullptr;
+}
diff --git a/Engine/Source/Runtime/ActorAnimation/Private/ActorAnimationPrivatePCH.h b/Engine/Source/Runtime/ActorAnimation/Private/ActorAnimationPrivatePCH.h
new file mode 100644
--- /dev/null
+++ b/Engine/Source/Runtime/ActorAnimation/Private/ActorAnimationPrivatePCH.h
@@ -0,0 +1,11 @@
+// Copyright 1998-2015 Epic Games, Inc. All Rights Reserved.
+
+#pragma once
+
+
+/* Private dependencies
+ *****************************************************************************/
+
+#include "Core.h"
+#include "CoreUObject.h"
+#include "Engine.h"
diff --git a/Engine/Source/Runtime/ActorAnimation/Public/ActorAnimation.h b/Engine/Source/Runtime/ActorAnimation/Public/ActorAnimation.h
new file mode 100644
--- /dev/null
+++ b/Engine/Source/Runtime/ActorAnimation/Public/ActorAnimation.h
@@ -0,0 +1,65 @@
+// Copyright 1998-2015 Epic Games, Inc. All Rights Reserved.
+
+#pragma once
+
+#include "MovieSceneAnimation.h"
+#include "ActorAnimationObject.h"
+#include "ActorAnimation.generated.h"
+
+
+/**
+ * Movie scene animation for Actors.
+ */
+UCLASS(BlueprintType, MinimalAPI)
+class UActorAnimation
+	: public UObject
+	, public IMovieSceneAnimation
+{
+	GENERATED_UCLASS_BODY()
+
+	/** Initialize this actor animation. */
+	ACTORANIMATION_API void Initialize();
+
+public:
+
+	// IMovieSceneAnimation interface
+
+	virtual bool AllowsSpawnableObjects() const override;
+	virtual void BindPossessableObject(const FGuid& ObjectId, UObject& PossessedObject) override;
+	virtual bool CanPossessObject(UObject& Object) const override;
+	virtual void DestroyAllSpawnedObjects(UMovieScene& MovieScene) override;
+	virtual UObject* FindObject(const FGuid& ObjectId) const override;
+	virtual FGuid FindObjectId(UObject& Object) const override;
+	virtual UMovieScene* GetMovieScene() const override;
+	virtual UObject* GetParentObject(UObject* Object) const override;
+	virtual void SpawnOrDestroyObjects(UMovieScene* MovieScene, bool DestroyAll) override;
+	virtual void UnbindPossessableObjects(const FGuid& ObjectId) override;
+
+#if WITH_EDITOR
+	virtual bool TryGetObjectDisplayName(const FGuid& ObjectId, FText& OutDisplayName) const override;
+#endif
+
+protected:
+
+	/**
+	 * Deselect all proxy actors in the Editor.
+	 *
+	 * @todo sequencer: remove Editor dependencies from this class
+	 */
+	ACTORANIMATION_API void DeselectAllActors();
+
+private:
+
+	/** Pointer to the movie scene that controls this animation. */
+	UPROPERTY()
+	UMovieScene* MovieScene;
+
+	/** Collection of possessed objects. */
+	// @todo sequencer: gmp: need support for UStruct keys in TMap
+	UPROPERTY()
+	TMap<FString, FActorAnimationObject> PossessedObjects;
+	//TMap<FGuid, FSequencerBinding> Bindings;
+
+	/** Collection of spawned proxy actors. */
+	TMap<FGuid, TWeakObjectPtr<AActor>> SpawnedActors;
+};
diff --git a/Engine/Source/Runtime/ActorAnimation/Public/ActorAnimationObject.h b/Engine/Source/Runtime/ActorAnimation/Public/ActorAnimationObject.h
new file mode 100644
--- /dev/null
+++ b/Engine/Source/Runtime/ActorAnimation/Public/ActorAnimationObject.h
@@ -0,0 +1,78 @@
+// Copyright 1998-2015 Epic Games, Inc. All Rights Reserved.
+
+#pragma once
+
+#include "ActorAnimationObject.generated.h"
+
+
+/**
+ * Structure for animated Actor objects.
+ */
+USTRUCT()
+struct FActorAnimationObject
+{
+	GENERATED_USTRUCT_BODY(FActorAnimationObject)
+
+public:
+
+	/** Creates and initializes a new instance. */
+	FActorAnimationObject()
+		: ObjectOrOwner(nullptr)
+		, CachedComponent(nullptr)
+	{ }
+
+	/**
+	 * Creates and initializes a new instance from an object.
+	 *
+	 * @param InObject The object to be bound.
+	 */
+	FActorAnimationObject(UObject* InObject)
+		: ObjectOrOwner(InObject)
+		, CachedComponent(nullptr)
+	{ }
+
+	/**
+	 * Creates and initializes a new instance from an object component.
+	 *
+	 * @param InOwner The object that owns the component.
+	 * @param InComponentName The component to be bound.
+	 */
+	FActorAnimationObject(UObject* InOwner, FString InComponentName)
+		: ObjectOrOwner(InOwner)
+		, ComponentName(InComponentName)
+		, CachedComponent(nullptr)
+	{ }
+
+	/**
+	 * Compares two bindings for equality.
+	 *
+	 * @param X The first binding to compare.
+	 * @param Y The second binding to compare.
+	 * @return true if the bindings refer to the same object, false otherwise.
+	 */
+	friend bool operator==(const FActorAnimationObject& X, const FActorAnimationObject& Y)
+	{
+		return (X.ObjectOrOwner == Y.ObjectOrOwner) && (X.ComponentName == Y.ComponentName);
+	}
+
+	/**
+	 * Gets a pointer to the possessed object.
+	 *
+	 * @return The object (usually an Actor or an ActorComponent).
+	 */
+	ACTORANIMATION_API UObject* GetObject() const;
+
+private:
+
+	/** The object or the owner of the object being possessed. */
+	UPROPERTY()
+	TLazyObjectPtr<UObject> ObjectOrOwner;
+
+	/** Optional name of an ActorComponent. */
+	UPROPERTY()
+	FString ComponentName;
+
+	/** Cached pointer to the Actor component (only if ComponentName is set). */
+	UPROPERTY(transient)
+	mutable UObject* CachedComponent;
+};
diff --git a/Engine/Source/Runtime/MovieScene/Private/MovieSceneAnimation.cpp b/Engine/Source/Runtime/MovieScene/Private/MovieSceneAnimation.cpp
new file mode 100644
--- /dev/null
+++ b/Engine/Source/Runtime/MovieScene/Private/MovieSceneAnimation.cpp
@@ -0,0 +1,12 @@
+// Copyright 1998-2015 Epic Games, Inc. All Rights Reserved.
+
+#include "MovieScenePrivatePCH.h"
+#include "MovieSceneAnimation.h"
+
+
+/* UMovieSceneAnimation structors
+ *****************************************************************************/
+
+UMovieSceneAnimation::UMovieSceneAnimation(const FObjectInitializer& ObjectInitializer)
+	: Super(ObjectInitializer)
+{ }
diff --git a/Engine/Source/Runtime/MovieScene/Private/MovieSceneBindings.cpp b/Engine/Source/Runtime/MovieScene/Private/MovieSceneBindings.cpp
deleted file mode 100644
diff --git a/Engine/Source/Runtime/MovieScene/Private/MovieSceneInstance.cpp b/Engine/Source/Runtime/MovieScene/Private/MovieSceneInstance.cpp
--- a/Engine/Source/Runtime/MovieScene/Private/MovieSceneInstance.cpp
+++ b/Engine/Source/Runtime/MovieScene/Private/MovieSceneInstance.cpp
@@ -11,12 +11,6 @@ FMovieSceneInstance::FMovieSceneInstance( UMovieScene& InMovieScene )
 }
 
 
-TScriptInterface<UMovieSceneObjectManager> FMovieSceneInstance::GetObjectManager() const
-{
-	return MovieScene.Get()->GetObjectManager();
-}
-
-
 void FMovieSceneInstance::SaveState()
 {
 	TMap<FGuid, FMovieSceneObjectBindingInstance>::TIterator ObjectIt = ObjectBindingInstances.CreateIterator();
diff --git a/Engine/Source/Runtime/MovieScene/Private/MovieSceneModule.cpp b/Engine/Source/Runtime/MovieScene/Private/MovieSceneModule.cpp
--- a/Engine/Source/Runtime/MovieScene/Private/MovieSceneModule.cpp
+++ b/Engine/Source/Runtime/MovieScene/Private/MovieSceneModule.cpp
@@ -1,10 +1,6 @@
 // Copyright 1998-2015 Epic Games, Inc. All Rights Reserved.
 
 #include "MovieScenePrivatePCH.h"
-#include "ModuleManager.h"
-#include "MovieScene.h"
-#include "MovieSceneBindings.h"
-#include "RuntimeMovieScenePlayer.h"
 
 
 DEFINE_LOG_CATEGORY(LogSequencerRuntime);
diff --git a/Engine/Source/Runtime/MovieScene/Private/MovieSceneObjectManager.cpp b/Engine/Source/Runtime/MovieScene/Private/MovieSceneObjectManager.cpp
deleted file mode 100644
diff --git a/Engine/Source/Runtime/MovieScene/Private/RuntimeMovieScenePlayer.cpp b/Engine/Source/Runtime/MovieScene/Private/RuntimeMovieScenePlayer.cpp
deleted file mode 100644
diff --git a/Engine/Source/Runtime/MovieScene/Public/MovieScene.h b/Engine/Source/Runtime/MovieScene/Public/MovieScene.h
--- a/Engine/Source/Runtime/MovieScene/Public/MovieScene.h
+++ b/Engine/Source/Runtime/MovieScene/Public/MovieScene.h
@@ -9,7 +9,6 @@
 
 
 class UBlueprint;
-class UMovieSceneObjectManager;
 class UMovieSceneSection;
 class UMovieSceneTrack;
 
@@ -239,16 +238,6 @@ public:
 		return ObjectBindings;
 	}
 
-	/**
-	 * Get the movie scene's object manager.
-	 *
-	 * @return The object manager.
-	 */
-	TScriptInterface<UMovieSceneObjectManager> GetObjectManager() const
-	{
-		return BindingManager;
-	}
-
 	/**
 	 * @return The time range of the movie scene (defined by all sections in the scene)
 	 */
@@ -282,10 +271,6 @@ protected:
 
 private:
 
-	/** The object binding manager. */
-	UPROPERTY()
-	TScriptInterface<UMovieSceneObjectManager> BindingManager;
-
 	/**
 	 * Data-only blueprints for all of the objects that we we're able to spawn.
 	 * These describe objects and actors that we may instantiate at runtime,
diff --git a/Engine/Source/Runtime/MovieScene/Public/MovieSceneAnimation.h b/Engine/Source/Runtime/MovieScene/Public/MovieSceneAnimation.h
new file mode 100644
--- /dev/null
+++ b/Engine/Source/Runtime/MovieScene/Public/MovieSceneAnimation.h
@@ -0,0 +1,125 @@
+// Copyright 1998-2015 Epic Games, Inc. All Rights Reserved.
+
+#pragma once
+
+#include "MovieSceneAnimation.generated.h"
+
+
+class FMovieSceneInstance;
+class UMovieScene;
+class UObject;
+
+
+/**
+ * Interface for movie scene animations (UObject version).
+ */
+UINTERFACE(MinimalAPI, meta=(CannotImplementInterfaceInBlueprint))
+class UMovieSceneAnimation
+	: public UInterface
+{
+	GENERATED_UINTERFACE_BODY()
+};
+
+
+/**
+ * Interface for movie scene animations (C++ version).
+ */
+class IMovieSceneAnimation
+{
+	GENERATED_IINTERFACE_BODY()
+
+	/**
+	 * Whether objects can be spawned at run-time.
+	 *
+	 * @return true if objects can be spawned by sequencer, false if only existing objects can be possessed.
+	 */
+	virtual bool AllowsSpawnableObjects() const PURE_VIRTUAL(IMovieSceneAnimation::AllowsSpawnableObjects, return false;);
+
+	/**
+	 * Called when Sequencer has created an object binding for a possessable object
+	 * 
+	 * @param ObjectId The guid used to map to the possessable object.  Note the guid can be bound to multiple objects at once
+	 * @param PossessedObject The runtime object which was possessed.
+	 * @see UnbindPossessableObjects
+	 */
+	virtual void BindPossessableObject(const FGuid& ObjectId, UObject& PossessedObject) PURE_VIRTUAL(IMovieSceneAnimation::BindPossessableObject,);
+	
+	/**
+	 * Check whether the given object can be possessed by this object manager.
+	 *
+	 * @param Object The object to check.
+	 * @return true if the object can be possessed, false otherwise.
+	 */
+	virtual bool CanPossessObject(UObject& Object) const PURE_VIRTUAL(IMovieSceneAnimation::CanPossessObject, return false;);
+
+	/**
+	 * Destroy all objects that were spawned by this binding manager.
+	 *
+	 * @param MovieScene The movie scene that owns this manager.
+	 */
+	virtual void DestroyAllSpawnedObjects(UMovieScene& MovieScene) PURE_VIRTUAL(IMovieSceneAnimation::DestroyAllSpawnedObjects,);
+
+	/**
+	 * Gets the possessed or spawned object for the specified identifier.
+	 *
+	 * @param ObjectId The unique identifier of the object.
+	 * @return The object, or nullptr if not found.
+	 */
+	virtual UObject* FindObject(const FGuid& ObjectId) const PURE_VIRTUAL(IMovieSceneAnimation::FindObject, return nullptr;);
+
+	/**
+	 * Find the identifier for a possessed or spawned object.
+	 *
+	 * @param The object to get the identifier for.
+	 * @return The object identifier, or an invalid GUID if not found.
+	 */
+	virtual FGuid FindObjectId( UObject& Object) const PURE_VIRTUAL(IMovieSceneAnimation::FindGuidForObject, return FGuid(););
+
+	/**
+	 * Get the movie scene that controls this animation.
+	 *
+	 * The returned movie scene represents the root movie scene.
+	 * It may contain additional child movie scenes.
+	 *
+	 * @return The movie scene.
+	 */
+	virtual UMovieScene* GetMovieScene() const PURE_VIRTUAL(IMovieSceneAnimation::GetMovieScene(), return nullptr;);
+
+	/**
+	 * Get the logical parent object for the supplied object (not necessarily its outer).
+	 *
+	 * @param Object The object whose parent to get.
+	 * @return The parent object, or nullptr if the object has no logical parent.
+	 */
+	virtual UObject* GetParentObject(UObject* Object) const = 0;
+
+	/**
+	 * Called to spawn or destroy objects for a movie instance.
+	 *
+	 * @param MovieScene The movie scene to spawn or destroy objects for.
+	 * @param DestroyAll If true, destroy all spawned objects for the instance, if false only destroy unused objects.
+	 */
+	virtual void SpawnOrDestroyObjects(UMovieScene* MovieScene, bool DestroyAll) PURE_VIRTUAL(IMovieSceneAnimation::SpawnOrDestroyObjects,);
+
+	/**
+	 * Unbinds all possessable objects from the provided GUID.
+	 *
+	 * @param ObjectId The guid bound to possessable objects that should be removed.
+	 * @see BindPossessableObject
+	 */
+	virtual void UnbindPossessableObjects(const FGuid& ObjectId) PURE_VIRTUAL(IMovieSceneAnimation::UnbindPossessableObjects,);
+
+public:
+
+#if WITH_EDITOR
+	/**
+	 * 
+	 * Tries to get a display name for the binding represented by the guid.
+	 *
+	 * @param ObjectGuid The guid for the object binding.
+	 * @param DisplayName The display name for the object binding.
+	 * @returns true if DisplayName has been set to a valid display name, otherwise false.
+	 */
+	virtual bool TryGetObjectDisplayName(const FGuid& ObjectId, FText& OutDisplayName) const PURE_VIRTUAL(IMovieSceneAnimation::TryGetObjectDisplayName, return false;);
+#endif
+};
diff --git a/Engine/Source/Runtime/MovieScene/Public/MovieSceneBindings.h b/Engine/Source/Runtime/MovieScene/Public/MovieSceneBindings.h
deleted file mode 100644
diff --git a/Engine/Source/Runtime/MovieScene/Public/MovieSceneInstance.h b/Engine/Source/Runtime/MovieScene/Public/MovieSceneInstance.h
--- a/Engine/Source/Runtime/MovieScene/Public/MovieSceneInstance.h
+++ b/Engine/Source/Runtime/MovieScene/Public/MovieSceneInstance.h
@@ -3,7 +3,7 @@
 #pragma once
 
 
-class UMovieSceneObjectManager;
+class IMovieSceneTrackInstance;
 class UMovieSceneTrack;
 
 
@@ -67,13 +67,6 @@ public:
 		return TimeRange;
 	}
 
-	/**
-	 * Get the movie scene's object manager.
-	 *
-	 * @return The object manager.
-	 */
-	TScriptInterface<UMovieSceneObjectManager> GetObjectManager() const;
-
 protected:
 
 	void RefreshInstanceMap(const TArray<UMovieSceneTrack*>& Tracks, const TArray<UObject*>& RuntimeObjects, FMovieSceneInstanceMap& TrackInstances, class IMovieScenePlayer& Player);
@@ -98,7 +91,7 @@ private:
 	TWeakObjectPtr<UMovieScene> MovieScene;
 
 	/** The shot track instance map */
-	TSharedPtr<class IMovieSceneTrackInstance> ShotTrackInstance;
+	TSharedPtr<IMovieSceneTrackInstance> ShotTrackInstance;
 
 	/** All Master track instances */
 	FMovieSceneInstanceMap MasterTrackInstances;
diff --git a/Engine/Source/Runtime/MovieScene/Public/MovieSceneObjectManager.h b/Engine/Source/Runtime/MovieScene/Public/MovieSceneObjectManager.h
deleted file mode 100644
diff --git a/Engine/Source/Runtime/MovieScene/Public/RuntimeMovieScenePlayer.h b/Engine/Source/Runtime/MovieScene/Public/RuntimeMovieScenePlayer.h
deleted file mode 100644
diff --git a/Engine/Source/Runtime/Niagara/Private/NiagaraAnimation.cpp b/Engine/Source/Runtime/Niagara/Private/NiagaraAnimation.cpp
new file mode 100644
--- /dev/null
+++ b/Engine/Source/Runtime/Niagara/Private/NiagaraAnimation.cpp
@@ -0,0 +1,80 @@
+// Copyright 1998-2015 Epic Games, Inc. All Rights Reserved.
+
+#include "NiagaraPrivate.h"
+#include "NiagaraAnimation.h"
+
+
+/* UNiagaraAnimation structors
+ *****************************************************************************/
+
+UNiagaraAnimation::UNiagaraAnimation(const FObjectInitializer& ObjectInitializer)
+	: Super(ObjectInitializer)
+	, MovieScene(nullptr)
+{ }
+
+
+/* IMovieSceneAnimation interface
+ *****************************************************************************/
+
+bool UNiagaraAnimation::AllowsSpawnableObjects() const
+{
+	return false;
+}
+
+
+void UNiagaraAnimation::BindPossessableObject(const FGuid& ObjectId, UObject& PossessedObject)
+{
+}
+
+
+bool UNiagaraAnimation::CanPossessObject(UObject& Object) const
+{
+	return false;
+}
+
+
+void UNiagaraAnimation::DestroyAllSpawnedObjects(UMovieScene& MovieScene)
+{
+}
+
+
+UObject* UNiagaraAnimation::FindObject(const FGuid& ObjectId) const
+{
+	return nullptr;
+}
+
+
+FGuid UNiagaraAnimation::FindObjectId(UObject& Object) const
+{
+	return FGuid();
+}
+
+
+UMovieScene* UNiagaraAnimation::GetMovieScene() const
+{
+	return MovieScene;
+}
+
+
+UObject* UNiagaraAnimation::GetParentObject(UObject* Object) const
+{
+	return nullptr;
+}
+
+
+void UNiagaraAnimation::SpawnOrDestroyObjects(UMovieScene* MovieScene, bool DestroyAll)
+{
+}
+
+
+#if WITH_EDITOR
+bool UNiagaraAnimation::TryGetObjectDisplayName(const FGuid& ObjectId, FText& OutDisplayName) const
+{
+	return false;
+}
+#endif
+
+
+void UNiagaraAnimation::UnbindPossessableObjects(const FGuid& ObjectId)
+{
+}
diff --git a/Engine/Source/Runtime/Niagara/Public/NiagaraAnimation.h b/Engine/Source/Runtime/Niagara/Public/NiagaraAnimation.h
new file mode 100644
--- /dev/null
+++ b/Engine/Source/Runtime/Niagara/Public/NiagaraAnimation.h
@@ -0,0 +1,43 @@
+// Copyright 1998-2015 Epic Games, Inc. All Rights Reserved.
+
+#pragma once
+
+#include "MovieSceneAnimation.h"
+#include "NiagaraAnimation.generated.h"
+
+
+/**
+ * Movie scene animation used by Niagara.
+ */
+UCLASS(BlueprintType, MinimalAPI)
+class UNiagaraAnimation
+	: public UObject
+	, public IMovieSceneAnimation
+{
+	GENERATED_UCLASS_BODY()
+
+public:
+
+	// IMovieSceneAnimation interface
+
+	virtual bool AllowsSpawnableObjects() const override;
+	virtual void BindPossessableObject(const FGuid& ObjectId, UObject& PossessedObject) override;
+	virtual bool CanPossessObject(UObject& Object) const override;
+	virtual void DestroyAllSpawnedObjects(UMovieScene& MovieScene) override;
+	virtual UObject* FindObject(const FGuid& ObjectId) const override;
+	virtual FGuid FindObjectId(UObject& Object) const override;
+	virtual UMovieScene* GetMovieScene() const override;
+	virtual UObject* GetParentObject(UObject* Object) const override;
+	virtual void SpawnOrDestroyObjects(UMovieScene* MovieScene, bool DestroyAll) override;
+	virtual void UnbindPossessableObjects(const FGuid& ObjectId) override;
+
+#if WITH_EDITOR
+	virtual bool TryGetObjectDisplayName(const FGuid& ObjectId, FText& OutDisplayName) const override;
+#endif
+
+private:
+
+	/** Pointer to the movie scene that controls this animation. */
+	UPROPERTY()
+	UMovieScene* MovieScene;
+};
diff --git a/Engine/Source/Runtime/UMG/Private/Animation/UMGSequencePlayer.cpp b/Engine/Source/Runtime/UMG/Private/Animation/UMGSequencePlayer.cpp
--- a/Engine/Source/Runtime/UMG/Private/Animation/UMGSequencePlayer.cpp
+++ b/Engine/Source/Runtime/UMG/Private/Animation/UMGSequencePlayer.cpp
@@ -3,7 +3,6 @@
 #include "UMGPrivatePCH.h"
 #include "UMGSequencePlayer.h"
 #include "MovieScene.h"
-#include "MovieSceneBindings.h"
 #include "MovieSceneInstance.h"
 #include "MovieScene.h"
 #include "WidgetAnimation.h"
diff --git a/Engine/Source/Runtime/UMG/Private/Animation/WidgetAnimation.cpp b/Engine/Source/Runtime/UMG/Private/Animation/WidgetAnimation.cpp
--- a/Engine/Source/Runtime/UMG/Private/Animation/WidgetAnimation.cpp
+++ b/Engine/Source/Runtime/UMG/Private/Animation/WidgetAnimation.cpp
@@ -3,6 +3,10 @@
 #include "UMGPrivatePCH.h"
 #include "MovieScene.h"
 #include "WidgetAnimation.h"
+#include "WidgetTree.h"
+
+
+#define LOCTEXT_NAMESPACE "UWidgetAnimation"
 
 
 /* UWidgetAnimation structors
@@ -45,3 +49,235 @@ float UWidgetAnimation::GetEndTime() const
 {
 	return MovieScene->GetTimeRange().GetUpperBoundValue();
 }
+
+
+void UWidgetAnimation::Initialize(UUserWidget* InPreviewWidget)
+{
+	PreviewWidget = InPreviewWidget;
+
+	// clear object maps
+	PreviewObjectToIds.Empty();
+	IdToPreviewObjects.Empty();
+	IdToSlotContentPreviewObjects.Empty();
+	SlotContentPreviewObjectToIds.Empty();
+
+	if (PreviewWidget == nullptr)
+	{
+		return;
+	}
+
+	// populate object maps
+	UWidgetTree* WidgetTree = PreviewWidget->WidgetTree;
+
+	for (const FWidgetAnimationBinding& Binding : AnimationBindings)
+	{
+		UObject* FoundObject = Binding.FindRuntimeObject(*WidgetTree);
+
+		if (FoundObject == nullptr)
+		{
+			continue;
+		}
+
+		UPanelSlot* FoundSlot = Cast<UPanelSlot>(FoundObject);
+
+		if (FoundSlot == nullptr)
+		{
+			IdToPreviewObjects.Add(Binding.AnimationGuid, FoundObject);
+			PreviewObjectToIds.Add(FoundObject, Binding.AnimationGuid);
+		}
+		else
+		{
+			IdToSlotContentPreviewObjects.Add(Binding.AnimationGuid, FoundSlot->Content);
+			SlotContentPreviewObjectToIds.Add(FoundSlot->Content, Binding.AnimationGuid);
+		}
+	}
+}
+
+
+/* IMovieSceneAnimation interface
+ *****************************************************************************/
+
+bool UWidgetAnimation::AllowsSpawnableObjects() const
+{
+	return false;
+}
+
+
+void UWidgetAnimation::BindPossessableObject(const FGuid& ObjectId, UObject& PossessedObject)
+{
+	UPanelSlot* PossessedSlot = Cast<UPanelSlot>(&PossessedObject);
+
+	if ((PossessedSlot != nullptr) && (PossessedSlot->Content != nullptr))
+	{
+		SlotContentPreviewObjectToIds.Add(PossessedSlot->Content, ObjectId);
+		IdToSlotContentPreviewObjects.Add(ObjectId, PossessedSlot->Content);
+
+		// Save the name of the widget containing the slots. This is the object
+		// to look up that contains the slot itself (the thing we are animating).
+		FWidgetAnimationBinding NewBinding;
+		{
+			NewBinding.AnimationGuid = ObjectId;
+			NewBinding.SlotWidgetName = PossessedSlot->GetFName();
+			NewBinding.WidgetName = PossessedSlot->Content->GetFName();
+		}
+
+		AnimationBindings.Add(NewBinding);
+	}
+	else if (PossessedSlot == nullptr)
+	{
+		PreviewObjectToIds.Add(&PossessedObject, ObjectId);
+		IdToPreviewObjects.Add(ObjectId, &PossessedObject);
+
+		FWidgetAnimationBinding NewBinding;
+		{
+			NewBinding.AnimationGuid = ObjectId;
+			NewBinding.WidgetName = PossessedObject.GetFName();
+		}
+
+		AnimationBindings.Add(NewBinding);
+	}
+}
+
+
+bool UWidgetAnimation::CanPossessObject(UObject& Object) const
+{
+	UPanelSlot* Slot = Cast<UPanelSlot>(&Object);
+
+	if ((Slot != nullptr) && (Slot->Content == nullptr))
+	{
+		// can't possess empty slots.
+		return false;
+	}
+
+	return (Object.IsA<UVisual>() && Object.IsIn(PreviewWidget.Get()));
+}
+
+
+FGuid UWidgetAnimation::FindObjectId(UObject& Object) const
+{
+	UPanelSlot* Slot = Cast<UPanelSlot>(&Object);
+
+	if (Slot != nullptr)
+	{
+		// slot guids are tracked by their content.
+		return PreviewObjectToIds.FindRef(Slot->Content);
+	}
+
+	return PreviewObjectToIds.FindRef(&Object);
+}
+
+
+UObject* UWidgetAnimation::FindObject(const FGuid& ObjectId) const
+{
+	TWeakObjectPtr<UObject> PreviewObject = IdToPreviewObjects.FindRef(ObjectId);
+
+	if (PreviewObject.IsValid())
+	{
+		return PreviewObject.Get();
+	}
+
+	TWeakObjectPtr<UObject> SlotContentPreviewObject = IdToSlotContentPreviewObjects.FindRef(ObjectId);
+
+	if (SlotContentPreviewObject.IsValid())
+	{
+		return SlotContentPreviewObject.Get();
+	}
+	
+	return nullptr;
+}
+
+
+UMovieScene* UWidgetAnimation::GetMovieScene() const
+{
+	return MovieScene;
+}
+
+
+UObject* UWidgetAnimation::GetParentObject(UObject* Object) const
+{
+	UPanelSlot* Slot = Cast<UPanelSlot>(Object);
+
+	if (Slot != nullptr)
+	{
+		// The slot is actually the child of the panel widget in the hierarchy,
+		// but we want it to show up as a sub-object of the widget it contains
+		// in the timeline so we return the content's GUID.
+		return Slot->Content;
+	}
+
+	return nullptr;
+}
+
+
+#if WITH_EDITOR
+bool UWidgetAnimation::TryGetObjectDisplayName(const FGuid& ObjectId, FText& OutDisplayName) const
+{
+	// TODO: This gets called every frame for every bound object and could
+	// be a potential performance issue for a really complicated animation.
+
+	TWeakObjectPtr<UObject> PreviewObject = IdToPreviewObjects.FindRef(ObjectId);
+
+	if (PreviewObject.IsValid())
+	{
+		OutDisplayName = FText::FromString(PreviewObject.Get()->GetName());
+		return true;
+	}
+
+	TWeakObjectPtr<UObject> SlotContentPreviewObject = IdToSlotContentPreviewObjects.FindRef(ObjectId);
+
+	if (SlotContentPreviewObject.IsValid())
+	{
+		UWidget* SlotContent = Cast<UWidget>(SlotContentPreviewObject.Get());
+		FText PanelName = SlotContent->Slot != nullptr && SlotContent->Slot->Parent != nullptr
+			? FText::FromString(SlotContent->Slot->Parent->GetName())
+			: LOCTEXT("InvalidPanel", "Invalid Panel");
+		FText ContentName = FText::FromString(SlotContent->GetName());
+		OutDisplayName = FText::Format(LOCTEXT("SlotObject", "{0} ({1} Slot)"), ContentName, PanelName);
+
+		return true;
+	}
+
+	return false;
+}
+#endif
+
+
+void UWidgetAnimation::UnbindPossessableObjects(const FGuid& ObjectId)
+{
+	// unbind widgets
+	TArray<TWeakObjectPtr<UObject>> PreviewObjects;
+	{
+		IdToPreviewObjects.MultiFind(ObjectId, PreviewObjects);
+
+		for (TWeakObjectPtr<UObject>& PreviewObject : PreviewObjects)
+		{
+			PreviewObjectToIds.Remove(PreviewObject);
+		}
+
+		IdToPreviewObjects.Remove(ObjectId);
+	}
+
+	// unbind slots
+	TArray<TWeakObjectPtr<UObject>> SlotContentPreviewObjects;
+	{
+		IdToSlotContentPreviewObjects.MultiFind(ObjectId, SlotContentPreviewObjects);
+
+		for (TWeakObjectPtr<UObject>& SlotContentPreviewObject : SlotContentPreviewObjects)
+		{
+			SlotContentPreviewObjectToIds.Remove(SlotContentPreviewObject);
+		}
+
+		IdToSlotContentPreviewObjects.Remove(ObjectId);
+	}
+
+	// mark dirty
+	Modify();
+
+	// remove animation bindings
+	AnimationBindings.RemoveAll([&](const FWidgetAnimationBinding& Binding) {
+		return Binding.AnimationGuid == ObjectId;
+	});
+}
+
+
+#undef LOCTEXT_NAMESPACE
diff --git a/Engine/Source/Runtime/UMG/Public/Animation/WidgetAnimation.h b/Engine/Source/Runtime/UMG/Public/Animation/WidgetAnimation.h
--- a/Engine/Source/Runtime/UMG/Public/Animation/WidgetAnimation.h
+++ b/Engine/Source/Runtime/UMG/Public/Animation/WidgetAnimation.h
@@ -2,17 +2,19 @@
 
 #pragma once
 
-#include "IMovieScenePlayer.h"
+#include "MovieSceneAnimation.h"
 #include "WidgetAnimationBinding.h"
 #include "WidgetAnimation.generated.h"
 
 
 class UMovieScene;
+class UUserWidget;
 
 
 UCLASS(BlueprintType, MinimalAPI)
 class UWidgetAnimation
 	: public UObject
+	, public IMovieSceneAnimation
 {
 	GENERATED_UCLASS_BODY()
 
@@ -45,11 +47,49 @@ public:
 	UFUNCTION(BlueprintCallable, Category="Animation")
 	UMG_API float GetEndTime() const;
 
+	/**
+	 * Initialize the widge editor with a new user widget.
+	 *
+	 * @param InPreviewWidget The user widget to preview.
+	 */
+	UMG_API void Initialize(UUserWidget* InPreviewWidget);
+
+public:
+
+	// IMovieSceneAnimation interface
+
+	virtual bool AllowsSpawnableObjects() const override;
+	virtual void BindPossessableObject(const FGuid& ObjectId, UObject& PossessedObject) override;
+	virtual bool CanPossessObject(UObject& Object) const override;
+	virtual void DestroyAllSpawnedObjects(UMovieScene& MovieScene) override { }
+	virtual UObject* FindObject(const FGuid& ObjectId) const override;
+	virtual FGuid FindObjectId(UObject& Object) const override;
+	virtual UMovieScene* GetMovieScene() const override;
+	virtual UObject* GetParentObject(UObject* Object) const override;
+	virtual void SpawnOrDestroyObjects(UMovieScene* MovieScene, bool DestroyAll) override { }
+	virtual void UnbindPossessableObjects(const FGuid& ObjectId) override;
+
+#if WITH_EDITOR
+	virtual bool TryGetObjectDisplayName(const FGuid& ObjectId, FText& OutDisplayName) const override;
+#endif
+
 public:
 
+	/** Pointer to the movie scene that controls this animation. */
 	UPROPERTY()
 	UMovieScene* MovieScene;
 
 	UPROPERTY()
 	TArray<FWidgetAnimationBinding> AnimationBindings;
+
+private:
+
+	/** The current preview widget, if any. */
+	TWeakObjectPtr<UUserWidget> PreviewWidget;
+
+	/** Mapping of preview objects to sequencer GUIDs */
+	TMultiMap<FGuid, TWeakObjectPtr<UObject>> IdToPreviewObjects;
+	TMultiMap<FGuid, TWeakObjectPtr<UObject>> IdToSlotContentPreviewObjects;
+	TMap<TWeakObjectPtr<UObject>, FGuid> PreviewObjectToIds;
+	TMap<TWeakObjectPtr<UObject>, FGuid> SlotContentPreviewObjectToIds;
 };
