Ugly fix for legacy parsing code. We should never have 2 different string formats for vector..

diff --git a/Engine/Source/Editor/KismetCompiler/Private/KismetCompilerVMBackend.cpp b/Engine/Source/Editor/KismetCompiler/Private/KismetCompilerVMBackend.cpp
--- a/Engine/Source/Editor/KismetCompiler/Private/KismetCompilerVMBackend.cpp
+++ b/Engine/Source/Editor/KismetCompiler/Private/KismetCompilerVMBackend.cpp
@@ -463,22 +463,33 @@ public:
 				if (StructProperty->Struct == VectorStruct)
 				{
 					FVector V = FVector::ZeroVector;
-					FDefaultValueHelper::ParseVector(Term->Name, /*out*/ V);
+					const bool bParsedUsingCustomFormat = FDefaultValueHelper::ParseVector(Term->Name, /*out*/ V);
+					if (!bParsedUsingCustomFormat)
+					{
+						StructProperty->ImportText(*Term->Name, &V, PPF_None, nullptr);
+					}
 					Writer << EX_VectorConst;
 					Writer << V;
-
 				}
 				else if (StructProperty->Struct == RotatorStruct)
 				{
 					FRotator R = FRotator::ZeroRotator;
-					FDefaultValueHelper::ParseRotator(Term->Name, /*out*/ R);
+					const bool bParsedUsingCustomFormat = FDefaultValueHelper::ParseRotator(Term->Name, /*out*/ R);
+					if (!bParsedUsingCustomFormat)
+					{
+						StructProperty->ImportText(*Term->Name, &R, PPF_None, nullptr);
+					}
 					Writer << EX_RotationConst;
 					Writer << R;
 				}
 				else if (StructProperty->Struct == TransformStruct)
 				{
 					FTransform T = FTransform::Identity;
-					T.InitFromString( Term->Name );
+					const bool bParsedUsingCustomFormat = T.InitFromString(Term->Name);
+					if (!bParsedUsingCustomFormat)
+					{
+						StructProperty->ImportText(*Term->Name, &T, PPF_None, nullptr);
+					}
 					Writer << EX_TransformConst;
 					Writer << T;
 				}
