Recreating UberGraphFrame for StubAfterFailure compilation is redundant

diff --git a/Engine/Source/Editor/KismetCompiler/Private/KismetCompilerModule.cpp b/Engine/Source/Editor/KismetCompiler/Private/KismetCompilerModule.cpp
--- a/Engine/Source/Editor/KismetCompiler/Private/KismetCompilerModule.cpp
+++ b/Engine/Source/Editor/KismetCompiler/Private/KismetCompilerModule.cpp
@@ -226,7 +226,6 @@ void FKismet2CompilerModule::CompileBlueprint(class UBlueprint* Blueprint, const
 			FKismetCompilerOptions StubCompileOptions(CompileOptions);
 			StubCompileOptions.CompileType = EKismetCompileType::StubAfterFailure;
 			{
-				FRecreateUberGraphFrameScope RecreateUberGraphFrameScope(Blueprint->GeneratedClass, bBytecodeOnly);
 				CompileBlueprintInner(Blueprint, StubCompileOptions, StubResults, StubReinstancer, ObjLoaded);
 			}
 
