UEBP-40 C++ code generation (WIP):
- skip redundant "switch"

diff --git a/Engine/Source/Developer/BlueprintCompilerCppBackend/Private/BlueprintCompilerCppBackendBase.cpp b/Engine/Source/Developer/BlueprintCompilerCppBackend/Private/BlueprintCompilerCppBackendBase.cpp
--- a/Engine/Source/Developer/BlueprintCompilerCppBackend/Private/BlueprintCompilerCppBackendBase.cpp
+++ b/Engine/Source/Developer/BlueprintCompilerCppBackend/Private/BlueprintCompilerCppBackendBase.cpp
@@ -288,34 +288,7 @@ void FBlueprintCompilerCppBackendBase::ConstructFunction(FKismetFunctionContext&
 		// Emit local variables
 		DeclareLocalVariables(FunctionContext, LocalVariables);
 
-		bool bUseSwitchState = false;
-		for (auto Node : FunctionContext.LinearExecutionList)
-		{
-			TArray<FBlueprintCompiledStatement*>* StatementList = FunctionContext.StatementsPerNode.Find(Node);
-			if (StatementList)
-			{
-				for (auto Statement : (*StatementList))
-				{
-					if (Statement && (
-						Statement->Type == KCST_UnconditionalGoto ||
-						Statement->Type == KCST_PushState ||
-						Statement->Type == KCST_GotoIfNot ||
-						Statement->Type == KCST_ComputedGoto ||
-						Statement->Type == KCST_EndOfThread ||
-						Statement->Type == KCST_EndOfThreadIfNot ||
-						Statement->Type == KCST_GotoReturn ||
-						Statement->Type == KCST_GotoReturnIfNot))
-					{
-						bUseSwitchState = true;
-						break;
-					}
-				}
-			}
-			if (bUseSwitchState)
-			{
-				break;
-			}
-		}
+		const bool bUseSwitchState = FunctionContext.MustUseSwitchState(nullptr);
 
 		// Run thru code looking only at things marked as jump targets, to make sure the jump targets are ordered in order of appearance in the linear execution list
 		// Emit code in the order specified by the linear execution list (the first node is always the entry point for the function)
diff --git a/Engine/Source/Editor/KismetCompiler/Private/KismetCompilerMisc.cpp b/Engine/Source/Editor/KismetCompiler/Private/KismetCompilerMisc.cpp
--- a/Engine/Source/Editor/KismetCompiler/Private/KismetCompilerMisc.cpp
+++ b/Engine/Source/Editor/KismetCompiler/Private/KismetCompilerMisc.cpp
@@ -1423,6 +1423,33 @@ int32 FKismetFunctionContext::GetContextUniqueID()
 	return UUIDCounter++;
 }
 
+bool FKismetFunctionContext::MustUseSwitchState(const FBlueprintCompiledStatement* ExcludeThisOne) const
+{
+	for (auto Node : LinearExecutionList)
+	{
+		auto StatementList = StatementsPerNode.Find(Node);
+		if (StatementList)
+		{
+			for (auto Statement : (*StatementList))
+			{
+				if (Statement && (Statement != ExcludeThisOne) && (
+					Statement->Type == KCST_UnconditionalGoto ||
+					Statement->Type == KCST_PushState ||
+					Statement->Type == KCST_GotoIfNot ||
+					Statement->Type == KCST_ComputedGoto ||
+					Statement->Type == KCST_EndOfThread ||
+					Statement->Type == KCST_EndOfThreadIfNot ||
+					Statement->Type == KCST_GotoReturn ||
+					Statement->Type == KCST_GotoReturnIfNot))
+				{
+					return true;
+				}
+			}
+		}
+	}
+	return false;
+}
+
 void FKismetFunctionContext::MergeAdjacentStates()
 {
 	for (int32 ExecIndex = 0; ExecIndex < LinearExecutionList.Num(); ++ExecIndex)
@@ -1450,17 +1477,15 @@ void FKismetFunctionContext::MergeAdjacentStates()
 		}
 	}
 
-	if (!bGeneratingCpp)
+	// Remove unnecessary GotoReturn statements
+	// if it's last statement generated by last node (in LinearExecution) then it can be removed
+	const auto LastExecutedNode = LinearExecutionList.Num() ? LinearExecutionList.Last() : NULL;
+	TArray<FBlueprintCompiledStatement*>* StatementList = StatementsPerNode.Find(LastExecutedNode);
+	FBlueprintCompiledStatement* LastStatementInLastNode = (StatementList && StatementList->Num()) ? StatementList->Last() : NULL;
+	const bool SafeForNativeCode = !bGeneratingCpp || !MustUseSwitchState(LastStatementInLastNode);
+	if (LastStatementInLastNode && SafeForNativeCode && (KCST_GotoReturn == LastStatementInLastNode->Type) && !LastStatementInLastNode->bIsJumpTarget)
 	{
-		// Remove unnecessary GotoReturn statements
-		// if it's last statement generated by last node (in LinearExecution) then it can be removed
-		const auto LastExecutedNode = LinearExecutionList.Num() ? LinearExecutionList.Last() : NULL;
-		TArray<FBlueprintCompiledStatement*>* StatementList = StatementsPerNode.Find(LastExecutedNode);
-		FBlueprintCompiledStatement* LastStatementInLastNode = (StatementList && StatementList->Num()) ? StatementList->Last() : NULL;
-		if (LastStatementInLastNode && (KCST_GotoReturn == LastStatementInLastNode->Type) && !LastStatementInLastNode->bIsJumpTarget)
-		{
-			StatementList->RemoveAt(StatementList->Num() - 1);
-		}
+		StatementList->RemoveAt(StatementList->Num() - 1);
 	}
 }
 
diff --git a/Engine/Source/Editor/KismetCompiler/Public/KismetCompiledFunctionContext.h b/Engine/Source/Editor/KismetCompiler/Public/KismetCompiledFunctionContext.h
--- a/Engine/Source/Editor/KismetCompiler/Public/KismetCompiledFunctionContext.h
+++ b/Engine/Source/Editor/KismetCompiler/Public/KismetCompiledFunctionContext.h
@@ -277,6 +277,8 @@ public:
 		return (SourceStatementList != NULL) && (SourceStatementList->Num() > 0);
 	}
 
+	KISMETCOMPILER_API bool MustUseSwitchState(const FBlueprintCompiledStatement* ExcludeThisOne) const;
+
 private:
 	// Optimize out any useless jumps (jump to the very next statement, where the control flow can just fall through)
 	void MergeAdjacentStates();
