Anim Blueprint compiler can now access member variables directly instead of thunking to the VM

Optimization is enabled by default but is an engine setting in case it needs to be disabled (bOptimizeAnimBlueprintMemberVariableAccess).
Expanded FExposedValueHandler to also contain an array of source & dest properties (copy records) to copy between. These are resolved to read/write addresses on init & a simple memcpy at runtime instead of calling the usual generated custom event.
Custom events are not added at all if all properties use copy records. The event may still be added & called however as mixed-mode access is still supported (i.e. a thunk and a memcpy to different pin properties).
Unfortunately this required initialization of all ExposedValueHandlers (as I didnt want to add the branch overhead of lazy init) so all the anim nodes have needed a small change to call their base class Initialize().

diff --git a/Engine/Config/BaseEngine.ini b/Engine/Config/BaseEngine.ini
--- a/Engine/Config/BaseEngine.ini
+++ b/Engine/Config/BaseEngine.ini
@@ -131,6 +131,8 @@ DefaultSoundName=/Engine/EngineSounds/WhiteNoise.WhiteNoise
 MaximumLoopIterationCount = 1000000
 ; Allows Blueprint classes based on a native Actor or Component subclass to tick even if their parent has bCanEverTick flag set to false.
 bCanBlueprintsTickByDefault=true
+; Controls whether anim blueprint nodes that access member variables of their class directly should use the optimized path that avoids a thunk to the Blueprint VM
+bOptimizeAnimBlueprintMemberVariableAccess=true
 ; If setting bAllowDebugViewmodesOnConsoles=true, must also add EngineDebugMaterials to the StartupPackages for console platforms
 bAllowDebugViewmodesOnConsoles=false
 CameraRotationThreshold=45.0
diff --git a/Engine/Source/Editor/KismetCompiler/Private/AnimBlueprintCompiler.cpp b/Engine/Source/Editor/KismetCompiler/Private/AnimBlueprintCompiler.cpp
--- a/Engine/Source/Editor/KismetCompiler/Private/AnimBlueprintCompiler.cpp
+++ b/Engine/Source/Editor/KismetCompiler/Private/AnimBlueprintCompiler.cpp
@@ -121,6 +121,12 @@ UK2Node_CallFunction* FAnimBlueprintCompiler::SpawnCallAnimInstanceFunction(UEdG
 
 void FAnimBlueprintCompiler::CreateEvaluationHandler(UAnimGraphNode_Base* VisualAnimNode, FEvaluationHandlerRecord& Record)
 {
+	if(Record.OnlyUsesCopyRecords())
+	{
+		// simple property copies don't require wiring up to a custom event
+		return;
+	}
+
 	// Shouldn't create a handler if there is nothing to work with
 	check(Record.ServicedProperties.Num() > 0);
 	check(Record.NodeVariableProperty != NULL);
@@ -161,6 +167,11 @@ void FAnimBlueprintCompiler::CreateEvaluationHandler(UAnimGraphNode_Base* Visual
 		// Does it get serviced by this handler?
 		if (FAnimNodeSinglePropertyHandler* SourceInfo = Record.ServicedProperties.Find(PropertyName))
 		{
+			if (SourceInfo->SimpleCopyPropertyName != NAME_None)
+			{
+				continue;
+			}
+
 			if (TargetPin->PinType.bIsArray)
 			{
 				// Grab the array that we need to set members for
@@ -1129,7 +1140,7 @@ void FAnimBlueprintCompiler::CopyTermDefaultsToDefaultObject(UObject* DefaultObj
 	for (auto EvalLinkIt = ValidEvaluationHandlerList.CreateIterator(); EvalLinkIt; ++EvalLinkIt)
 	{
 		FEvaluationHandlerRecord& Record = *EvalLinkIt;
-		Record.PatchFunctionNameInto(DefaultObject);
+		Record.PatchFunctionNameAndCopyRecordsInto(DefaultObject);
 	}
 
 	// And patch in constant values that don't need to be re-evaluated every frame
diff --git a/Engine/Source/Editor/KismetCompiler/Private/AnimBlueprintCompiler.h b/Engine/Source/Editor/KismetCompiler/Private/AnimBlueprintCompiler.h
--- a/Engine/Source/Editor/KismetCompiler/Private/AnimBlueprintCompiler.h
+++ b/Engine/Source/Editor/KismetCompiler/Private/AnimBlueprintCompiler.h
@@ -50,8 +50,12 @@ protected:
 
 		UEdGraphPin* SinglePin;
 
+		// If this is valid we are a simple property copy, avoiding BP thunk
+		FName SimpleCopyPropertyName;
+
 		FAnimNodeSinglePropertyHandler()
 			: SinglePin(NULL)
+			, SimpleCopyPropertyName(NAME_None)
 		{
 		}
 	};
@@ -114,16 +118,83 @@ protected:
 		{
 		}
 
+		bool OnlyUsesCopyRecords() const
+		{
+			for(TMap<FName, FAnimNodeSinglePropertyHandler>::TConstIterator It(ServicedProperties); It; ++It)
+			{
+				const FAnimNodeSinglePropertyHandler& AnimNodeSinglePropertyHandler = It.Value();
+				if(AnimNodeSinglePropertyHandler.SimpleCopyPropertyName == NAME_None)
+				{
+					return false;
+				}
+			}
+
+			return true;
+		}
+
 		bool IsValid() const
 		{
 			return (NodeVariableProperty != NULL) && (EvaluationHandlerProperty != NULL);
 		}
 
-		void PatchFunctionNameInto(UObject* TargetObject) const
+		void PatchFunctionNameAndCopyRecordsInto(UObject* TargetObject) const
 		{
-			EvaluationHandlerProperty->ContainerPtrToValuePtr<FExposedValueHandler>(
-				NodeVariableProperty->ContainerPtrToValuePtr<void>(TargetObject)
-			)->BoundFunction = HandlerFunctionName;
+			FExposedValueHandler* HandlerPtr = EvaluationHandlerProperty->ContainerPtrToValuePtr<FExposedValueHandler>(NodeVariableProperty->ContainerPtrToValuePtr<void>(TargetObject));
+			HandlerPtr->CopyRecords.Empty();
+
+			bool bOnlyUsesCopyRecords = true;
+			for(TMap<FName, FAnimNodeSinglePropertyHandler>::TConstIterator PropertyIt(ServicedProperties); PropertyIt; ++PropertyIt)
+			{
+				const FAnimNodeSinglePropertyHandler& AnimNodeSinglePropertyHandler = PropertyIt.Value();
+				if(AnimNodeSinglePropertyHandler.SimpleCopyPropertyName != NAME_None)
+				{
+					// Local variable get, no need to thunk to BP to grab this value
+
+					UProperty* SimpleCopyPropertySource = TargetObject->GetClass()->FindPropertyByName(AnimNodeSinglePropertyHandler.SimpleCopyPropertyName);
+					check(SimpleCopyPropertySource);
+
+					if(AnimNodeSinglePropertyHandler.ArrayPins.Num() > 0)
+					{
+						UArrayProperty* SimpleCopyPropertyDest = CastChecked<UArrayProperty>(NodeVariableProperty->Struct->FindPropertyByName(PropertyIt.Key()));
+						check(SimpleCopyPropertySource->GetSize() == SimpleCopyPropertyDest->Inner->GetSize());
+
+						for(TMap<int32, UEdGraphPin*>::TConstIterator ArrayIt(AnimNodeSinglePropertyHandler.ArrayPins); ArrayIt; ++ArrayIt)
+						{
+							FExposedValueCopyRecord CopyRecord;
+							CopyRecord.DestProperty = SimpleCopyPropertyDest;
+							CopyRecord.DestArrayIndex = ArrayIt.Key();
+							CopyRecord.SourceProperty = SimpleCopyPropertySource;
+							CopyRecord.SourceArrayIndex = 0;
+							CopyRecord.Size = SimpleCopyPropertyDest->Inner->GetSize();
+							HandlerPtr->CopyRecords.Add(CopyRecord);	
+						}
+					}
+					else
+					{
+						UProperty* SimpleCopyPropertyDest = NodeVariableProperty->Struct->FindPropertyByName(PropertyIt.Key());
+						check(SimpleCopyPropertySource->GetSize() == SimpleCopyPropertyDest->GetSize());
+
+						// Local variable get, no need to thunk to BP to grab this value
+						FExposedValueCopyRecord CopyRecord;
+						CopyRecord.DestProperty = SimpleCopyPropertyDest;
+						CopyRecord.DestArrayIndex = 0;
+						CopyRecord.SourceProperty = SimpleCopyPropertySource;
+						CopyRecord.SourceArrayIndex = 0;
+						CopyRecord.Size = SimpleCopyPropertyDest->GetSize();
+						HandlerPtr->CopyRecords.Add(CopyRecord);
+					}
+				}
+				else
+				{
+					bOnlyUsesCopyRecords = false;
+				}
+			}
+
+			if(!bOnlyUsesCopyRecords)
+			{
+				// not all of our pins use copy records so we will need to call our exposed value handler
+				HandlerPtr->BoundFunction = HandlerFunctionName;
+			}
 		}
 
 		void RegisterPin(UEdGraphPin* SourcePin, UProperty* AssociatedProperty, int32 AssociatedPropertyArrayIndex)
@@ -139,6 +210,15 @@ protected:
 				check(Handler.SinglePin == NULL);
 				Handler.SinglePin = SourcePin;
 			}
+
+			if(GetDefault<UEngine>()->bOptimizeAnimBlueprintMemberVariableAccess)
+			{
+				UK2Node_VariableGet* VariableGetNode = Cast<UK2Node_VariableGet>(SourcePin->LinkedTo[0]->GetOwningNode());
+				if(VariableGetNode && VariableGetNode->IsNodePure() && VariableGetNode->VariableReference.IsSelfContext())
+				{
+					Handler.SimpleCopyPropertyName = VariableGetNode->VariableReference.GetMemberName();
+				}
+			}
 		}
 	};
 
diff --git a/Engine/Source/Editor/UnrealEd/Private/EditorEngine.cpp b/Engine/Source/Editor/UnrealEd/Private/EditorEngine.cpp
--- a/Engine/Source/Editor/UnrealEd/Private/EditorEngine.cpp
+++ b/Engine/Source/Editor/UnrealEd/Private/EditorEngine.cpp
@@ -1770,6 +1770,17 @@ void UEditorEngine::PostEditChangeProperty(FPropertyChangedEvent& PropertyChange
 			Blueprint->Status = BS_Dirty;
 		}
 	}
+	else if (PropertyName == GET_MEMBER_NAME_CHECKED(UEngine, bOptimizeAnimBlueprintMemberVariableAccess))
+	{
+		FScopedSlowTask SlowTask(100, LOCTEXT("DirtyingAnimBlueprintsDueToOptimizationChange", "Invalidating All Anim Blueprints"));
+
+		// Flag all Blueprints as out of date (this doesn't dirty the package as needs saving but will force a recompile during PIE)
+		for (TObjectIterator<UAnimBlueprint> AnimBlueprintIt; AnimBlueprintIt; ++AnimBlueprintIt)
+		{
+			UAnimBlueprint* AnimBlueprint = *AnimBlueprintIt;
+			AnimBlueprint->Status = BS_Dirty;
+		}
+	}
 }
 
 void UEditorEngine::Cleanse( bool ClearSelection, bool Redraw, const FText& TransReset )
diff --git a/Engine/Source/Runtime/AnimGraphRuntime/Private/AnimNodes/AnimNode_ApplyAdditive.cpp b/Engine/Source/Runtime/AnimGraphRuntime/Private/AnimNodes/AnimNode_ApplyAdditive.cpp
--- a/Engine/Source/Runtime/AnimGraphRuntime/Private/AnimNodes/AnimNode_ApplyAdditive.cpp
+++ b/Engine/Source/Runtime/AnimGraphRuntime/Private/AnimNodes/AnimNode_ApplyAdditive.cpp
@@ -8,6 +8,8 @@
 
 void FAnimNode_ApplyAdditive::Initialize(const FAnimationInitializeContext& Context)
 {
+	FAnimNode_Base::Initialize(Context);
+
 	Base.Initialize(Context);
 	Additive.Initialize(Context);
 }
diff --git a/Engine/Source/Runtime/AnimGraphRuntime/Private/AnimNodes/AnimNode_BlendListBase.cpp b/Engine/Source/Runtime/AnimGraphRuntime/Private/AnimNodes/AnimNode_BlendListBase.cpp
--- a/Engine/Source/Runtime/AnimGraphRuntime/Private/AnimNodes/AnimNode_BlendListBase.cpp
+++ b/Engine/Source/Runtime/AnimGraphRuntime/Private/AnimNodes/AnimNode_BlendListBase.cpp
@@ -11,6 +11,8 @@
 
 void FAnimNode_BlendListBase::Initialize(const FAnimationInitializeContext& Context)
 {
+	FAnimNode_Base::Initialize(Context);
+
 	const int NumPoses = BlendPose.Num();
 	checkSlow(BlendTime.Num() == NumPoses);
 
diff --git a/Engine/Source/Runtime/AnimGraphRuntime/Private/AnimNodes/AnimNode_LayeredBoneBlend.cpp b/Engine/Source/Runtime/AnimGraphRuntime/Private/AnimNodes/AnimNode_LayeredBoneBlend.cpp
--- a/Engine/Source/Runtime/AnimGraphRuntime/Private/AnimNodes/AnimNode_LayeredBoneBlend.cpp
+++ b/Engine/Source/Runtime/AnimGraphRuntime/Private/AnimNodes/AnimNode_LayeredBoneBlend.cpp
@@ -10,6 +10,8 @@
 
 void FAnimNode_LayeredBoneBlend::Initialize(const FAnimationInitializeContext& Context)
 {
+	FAnimNode_Base::Initialize(Context);
+
 	const int NumPoses = BlendPoses.Num();
 	checkSlow(BlendWeights.Num() == NumPoses);
 
diff --git a/Engine/Source/Runtime/AnimGraphRuntime/Private/AnimNodes/AnimNode_Root.cpp b/Engine/Source/Runtime/AnimGraphRuntime/Private/AnimNodes/AnimNode_Root.cpp
--- a/Engine/Source/Runtime/AnimGraphRuntime/Private/AnimNodes/AnimNode_Root.cpp
+++ b/Engine/Source/Runtime/AnimGraphRuntime/Private/AnimNodes/AnimNode_Root.cpp
@@ -12,6 +12,8 @@ FAnimNode_Root::FAnimNode_Root()
 
 void FAnimNode_Root::Initialize(const FAnimationInitializeContext& Context)
 {
+	FAnimNode_Base::Initialize(Context);
+
 	Result.Initialize(Context);
 }
 
diff --git a/Engine/Source/Runtime/AnimGraphRuntime/Private/AnimNodes/AnimNode_RotateRootBone.cpp b/Engine/Source/Runtime/AnimGraphRuntime/Private/AnimNodes/AnimNode_RotateRootBone.cpp
--- a/Engine/Source/Runtime/AnimGraphRuntime/Private/AnimNodes/AnimNode_RotateRootBone.cpp
+++ b/Engine/Source/Runtime/AnimGraphRuntime/Private/AnimNodes/AnimNode_RotateRootBone.cpp
@@ -8,6 +8,8 @@
 
 void FAnimNode_RotateRootBone::Initialize(const FAnimationInitializeContext& Context)
 {
+	FAnimNode_Base::Initialize(Context);
+
 	BasePose.Initialize(Context);
 }
 
diff --git a/Engine/Source/Runtime/AnimGraphRuntime/Private/AnimNodes/AnimNode_SaveCachedPose.cpp b/Engine/Source/Runtime/AnimGraphRuntime/Private/AnimNodes/AnimNode_SaveCachedPose.cpp
--- a/Engine/Source/Runtime/AnimGraphRuntime/Private/AnimNodes/AnimNode_SaveCachedPose.cpp
+++ b/Engine/Source/Runtime/AnimGraphRuntime/Private/AnimNodes/AnimNode_SaveCachedPose.cpp
@@ -16,6 +16,8 @@ FAnimNode_SaveCachedPose::FAnimNode_SaveCachedPose()
 
 void FAnimNode_SaveCachedPose::Initialize(const FAnimationInitializeContext& Context)
 {
+	FAnimNode_Base::Initialize(Context);
+
 	if (LastInitializedContextCounter != Context.AnimInstance->GetGraphTraversalCounter())
 	{
 		LastInitializedContextCounter = Context.AnimInstance->GetGraphTraversalCounter();
diff --git a/Engine/Source/Runtime/AnimGraphRuntime/Private/AnimNodes/AnimNode_SequenceEvaluator.cpp b/Engine/Source/Runtime/AnimGraphRuntime/Private/AnimNodes/AnimNode_SequenceEvaluator.cpp
--- a/Engine/Source/Runtime/AnimGraphRuntime/Private/AnimNodes/AnimNode_SequenceEvaluator.cpp
+++ b/Engine/Source/Runtime/AnimGraphRuntime/Private/AnimNodes/AnimNode_SequenceEvaluator.cpp
@@ -8,6 +8,7 @@
 
 void FAnimNode_SequenceEvaluator::Initialize(const FAnimationInitializeContext& Context)
 {
+	FAnimNode_Base::Initialize(Context);
 }
 
 void FAnimNode_SequenceEvaluator::CacheBones(const FAnimationCacheBonesContext& Context) 
diff --git a/Engine/Source/Runtime/AnimGraphRuntime/Private/AnimNodes/AnimNode_Slot.cpp b/Engine/Source/Runtime/AnimGraphRuntime/Private/AnimNodes/AnimNode_Slot.cpp
--- a/Engine/Source/Runtime/AnimGraphRuntime/Private/AnimNodes/AnimNode_Slot.cpp
+++ b/Engine/Source/Runtime/AnimGraphRuntime/Private/AnimNodes/AnimNode_Slot.cpp
@@ -9,6 +9,8 @@
 
 void FAnimNode_Slot::Initialize(const FAnimationInitializeContext& Context)
 {
+	FAnimNode_Base::Initialize(Context);
+
 	Source.Initialize(Context);
 
 	SourceWeight = 0.f;
diff --git a/Engine/Source/Runtime/AnimGraphRuntime/Private/AnimNodes/AnimNode_TwoWayBlend.cpp b/Engine/Source/Runtime/AnimGraphRuntime/Private/AnimNodes/AnimNode_TwoWayBlend.cpp
--- a/Engine/Source/Runtime/AnimGraphRuntime/Private/AnimNodes/AnimNode_TwoWayBlend.cpp
+++ b/Engine/Source/Runtime/AnimGraphRuntime/Private/AnimNodes/AnimNode_TwoWayBlend.cpp
@@ -8,6 +8,8 @@
 
 void FAnimationNode_TwoWayBlend::Initialize(const FAnimationInitializeContext& Context)
 {
+	FAnimNode_Base::Initialize(Context);
+
 	A.Initialize(Context);
 	B.Initialize(Context);
 }
diff --git a/Engine/Source/Runtime/AnimGraphRuntime/Private/BoneControllers/AnimNode_SkeletalControlBase.cpp b/Engine/Source/Runtime/AnimGraphRuntime/Private/BoneControllers/AnimNode_SkeletalControlBase.cpp
--- a/Engine/Source/Runtime/AnimGraphRuntime/Private/BoneControllers/AnimNode_SkeletalControlBase.cpp
+++ b/Engine/Source/Runtime/AnimGraphRuntime/Private/BoneControllers/AnimNode_SkeletalControlBase.cpp
@@ -10,6 +10,8 @@
 
 void FAnimNode_SkeletalControlBase::Initialize(const FAnimationInitializeContext& Context)
 {
+	FAnimNode_Base::Initialize(Context);
+
 	ComponentPose.Initialize(Context);
 }
 
diff --git a/Engine/Source/Runtime/Engine/Classes/Animation/AnimNodeBase.h b/Engine/Source/Runtime/Engine/Classes/Animation/AnimNodeBase.h
--- a/Engine/Source/Runtime/Engine/Classes/Animation/AnimNodeBase.h
+++ b/Engine/Source/Runtime/Engine/Classes/Animation/AnimNodeBase.h
@@ -347,6 +347,38 @@ public:
 	void EvaluateComponentSpace(FComponentSpacePoseContext& Output);
 };
 
+USTRUCT()
+struct FExposedValueCopyRecord
+{
+	GENERATED_USTRUCT_BODY()
+
+	FExposedValueCopyRecord()
+		: Source(nullptr)
+		, Dest(nullptr)
+	{}
+
+	UPROPERTY()
+	UProperty* SourceProperty;
+
+	UPROPERTY()
+	int32 SourceArrayIndex;
+
+	UPROPERTY()
+	UProperty* DestProperty;
+
+	UPROPERTY()
+	int32 DestArrayIndex;
+
+	UPROPERTY()
+	int32 Size;
+
+	// Cached source copy ptr
+	void* Source;
+
+	// Cached dest copy ptr
+	void* Dest;
+};
+
 // An exposed value updater
 USTRUCT()
 struct FExposedValueHandler
@@ -357,17 +389,64 @@ struct FExposedValueHandler
 	UPROPERTY()
 	FName BoundFunction;
 
-	void Execute(const FAnimationBaseContext& Context) const
+	UPROPERTY(Transient)
+	UFunction* Function;
+
+	// Direct data access to property in anim instance
+	UPROPERTY()
+	TArray<FExposedValueCopyRecord> CopyRecords;
+
+	void Initialize(FAnimNode_Base* AnimNode, UAnimInstance* AnimInstance) 
 	{
 		if (BoundFunction != NAME_None)
 		{
-			//@TODO: Should be able to be Checked, or at least produce a warning when it fails
-			if (UFunction* Function = Context.AnimInstance->FindFunction(BoundFunction))
+			Function = AnimInstance->FindFunction(BoundFunction);
+		}
+
+		// initialize copy records
+		for(auto& CopyRecord : CopyRecords)
+		{
+			if(UArrayProperty* SourceArrayProperty = Cast<UArrayProperty>(CopyRecord.SourceProperty))
+			{
+				FScriptArrayHelper ArrayHelper(SourceArrayProperty, CopyRecord.SourceProperty->ContainerPtrToValuePtr<uint8>(AnimInstance));
+				check(ArrayHelper.IsValidIndex(CopyRecord.SourceArrayIndex));
+				CopyRecord.Source = ArrayHelper.GetRawPtr(CopyRecord.SourceArrayIndex);
+			}
+			else
+			{
+				CopyRecord.Source = CopyRecord.SourceProperty->ContainerPtrToValuePtr<uint8>(AnimInstance, CopyRecord.SourceArrayIndex);
+			}
+
+			if(UArrayProperty* DestArrayProperty = Cast<UArrayProperty>(CopyRecord.DestProperty))
+			{
+				FScriptArrayHelper ArrayHelper(DestArrayProperty, CopyRecord.DestProperty->ContainerPtrToValuePtr<uint8>(AnimNode));
+				check(ArrayHelper.IsValidIndex(CopyRecord.DestArrayIndex));
+				CopyRecord.Dest = ArrayHelper.GetRawPtr(CopyRecord.DestArrayIndex);
+			}
+			else
 			{
-				Context.AnimInstance->ProcessEvent(Function, NULL);
+				CopyRecord.Dest = CopyRecord.DestProperty->ContainerPtrToValuePtr<uint8>(AnimNode, CopyRecord.DestArrayIndex);
 			}
 		}
 	}
+
+	void Execute(const FAnimationBaseContext& Context) const
+	{
+		if (Function != nullptr)
+		{
+			Context.AnimInstance->ProcessEvent(Function, NULL);
+		}
+
+		for(const auto& CopyRecord : CopyRecords)
+		{
+			// if any of these checks fail then it's likely that Initialize has not been called.
+			// has new anim node type been added that doesnt call the base class Initialize()?
+			checkSlow(CopyRecord.Dest != nullptr);
+			checkSlow(CopyRecord.Source != nullptr);
+			checkSlow(CopyRecord.Size != 0);
+			FMemory::Memcpy(CopyRecord.Dest, CopyRecord.Source, CopyRecord.Size);
+		}
+	}
 };
 
 /**
@@ -389,7 +468,7 @@ struct ENGINE_API FAnimNode_Base
 	// A derived class should implement Initialize, Update, and either Evaluate or EvaluateComponentSpace, but not both of them
 
 	// Interface to implement
-	virtual void Initialize(const FAnimationInitializeContext& Context) {}
+	virtual void Initialize(const FAnimationInitializeContext& Context);
 	virtual void CacheBones(const FAnimationCacheBonesContext& Context) {}
 	virtual void Update(const FAnimationUpdateContext& Context) {}
 	virtual void Evaluate(FPoseContext& Output) { check(false); }
diff --git a/Engine/Source/Runtime/Engine/Classes/Engine/Engine.h b/Engine/Source/Runtime/Engine/Classes/Engine/Engine.h
--- a/Engine/Source/Runtime/Engine/Classes/Engine/Engine.h
+++ b/Engine/Source/Runtime/Engine/Classes/Engine/Engine.h
@@ -1154,6 +1154,10 @@ public:
 	UPROPERTY(EditAnywhere, config, Category=Blueprints)
 	uint32 bCanBlueprintsTickByDefault:1;
 
+	/** Controls whether anim blueprint nodes that access member variables of their class directly should use the optimized path that avoids a thunk to the Blueprint VM */
+	UPROPERTY(EditAnywhere, config, Category="Anim Blueprints")
+	uint32 bOptimizeAnimBlueprintMemberVariableAccess:1;
+
 	/** @todo document */
 	UPROPERTY(config)
 	uint32 bEnableEditorPSysRealtimeLOD:1;
diff --git a/Engine/Source/Runtime/Engine/Private/Animation/AnimNodeBase.cpp b/Engine/Source/Runtime/Engine/Private/Animation/AnimNodeBase.cpp
--- a/Engine/Source/Runtime/Engine/Private/Animation/AnimNodeBase.cpp
+++ b/Engine/Source/Runtime/Engine/Private/Animation/AnimNodeBase.cpp
@@ -3,6 +3,14 @@
 #include "EnginePrivate.h"
 #include "Animation/AnimNodeBase.h"
 
+/////////////////////////////////////////////////////
+// FAnimNode_Base
+
+void FAnimNode_Base::Initialize(const FAnimationInitializeContext& Context)
+{
+	EvaluateGraphExposedInputs.Initialize(this, Context.AnimInstance);
+}
+
 /////////////////////////////////////////////////////
 // FPoseLinkBase
 
diff --git a/Engine/Source/Runtime/Engine/Private/Animation/AnimNode_SequencePlayer.cpp b/Engine/Source/Runtime/Engine/Private/Animation/AnimNode_SequencePlayer.cpp
--- a/Engine/Source/Runtime/Engine/Private/Animation/AnimNode_SequencePlayer.cpp
+++ b/Engine/Source/Runtime/Engine/Private/Animation/AnimNode_SequencePlayer.cpp
@@ -8,6 +8,8 @@
 
 void FAnimNode_SequencePlayer::Initialize(const FAnimationInitializeContext& Context)
 {
+	FAnimNode_Base::Initialize(Context);
+
 	EvaluateGraphExposedInputs.Execute(Context);
 	InternalTimeAccumulator = StartPosition;
 	if (Sequence != NULL)
diff --git a/Engine/Source/Runtime/Engine/Private/Animation/AnimNode_StateMachine.cpp b/Engine/Source/Runtime/Engine/Private/Animation/AnimNode_StateMachine.cpp
--- a/Engine/Source/Runtime/Engine/Private/Animation/AnimNode_StateMachine.cpp
+++ b/Engine/Source/Runtime/Engine/Private/Animation/AnimNode_StateMachine.cpp
@@ -197,6 +197,8 @@ FBakedAnimationStateMachine* FAnimNode_StateMachine::GetMachineDescription()
 
 void FAnimNode_StateMachine::Initialize(const FAnimationInitializeContext& Context)
 {
+	FAnimNode_Base::Initialize(Context);
+
 	UAnimBlueprintGeneratedClass* AnimBlueprintClass = Context.GetAnimBlueprintClass();
 	PRIVATE_MachineDescription = AnimBlueprintClass->BakedStateMachines.IsValidIndex(StateMachineIndexInClass) ? &(AnimBlueprintClass->BakedStateMachines[StateMachineIndexInClass]) : NULL;
 
@@ -213,12 +215,34 @@ void FAnimNode_StateMachine::Initialize(const FAnimationInitializeContext& Conte
 			StatePoseLinks.Reserve(Machine->States.Num());
 			for (int32 StateIndex = 0; StateIndex < Machine->States.Num(); ++StateIndex)
 			{
+				const FBakedAnimationState& State = Machine->States[StateIndex];
 				FPoseLink* StatePoseLink = new (StatePoseLinks) FPoseLink();
 
 				// because conduits don't contain bound graphs, this link is no longer guaranteed to be valid
-				if (Machine->States[StateIndex].StateRootNodeIndex != INDEX_NONE)
+				if (State.StateRootNodeIndex != INDEX_NONE)
+				{
+					StatePoseLink->LinkID = AnimBlueprintClass->AnimNodeProperties.Num() - 1 - State.StateRootNodeIndex; //@TODO: Crazysauce
+				}
+
+				// also initialize transitions
+				if(State.EntryRuleNodeIndex != INDEX_NONE)
 				{
-					StatePoseLink->LinkID = AnimBlueprintClass->AnimNodeProperties.Num() - 1 - Machine->States[StateIndex].StateRootNodeIndex; //@TODO: Crazysauce
+					if (FAnimNode_TransitionResult* TransitionNode = GetNodeFromPropertyIndex<FAnimNode_TransitionResult>(Context.AnimInstance, AnimBlueprintClass, State.EntryRuleNodeIndex))
+					{
+						TransitionNode->Initialize(Context);
+					}
+				}
+
+				for(int32 TransitionIndex = 0; TransitionIndex < State.Transitions.Num(); ++TransitionIndex)
+				{
+					const FBakedStateExitTransition& TransitionRule = State.Transitions[TransitionIndex];
+					if (TransitionRule.CanTakeDelegateIndex != INDEX_NONE)
+					{
+						if (FAnimNode_TransitionResult* TransitionNode = GetNodeFromPropertyIndex<FAnimNode_TransitionResult>(Context.AnimInstance, AnimBlueprintClass, TransitionRule.CanTakeDelegateIndex))
+						{
+							TransitionNode->Initialize(Context);
+						}
+					}
 				}
 			}
 
diff --git a/Engine/Source/Runtime/Engine/Private/Animation/AnimNode_TransitionPoseEvaluator.cpp b/Engine/Source/Runtime/Engine/Private/Animation/AnimNode_TransitionPoseEvaluator.cpp
--- a/Engine/Source/Runtime/Engine/Private/Animation/AnimNode_TransitionPoseEvaluator.cpp
+++ b/Engine/Source/Runtime/Engine/Private/Animation/AnimNode_TransitionPoseEvaluator.cpp
@@ -16,6 +16,8 @@ FAnimNode_TransitionPoseEvaluator::FAnimNode_TransitionPoseEvaluator()
 
 void FAnimNode_TransitionPoseEvaluator::Initialize(const FAnimationInitializeContext& Context)
 {	
+	FAnimNode_Base::Initialize(Context);
+
 	if (EvaluatorMode == EEvaluatorMode::EM_Freeze)
 	{
 		// EM_Freeze must evaluate 1 frame to get the initial pose. This cached frame will not call update, only evaluate
diff --git a/Engine/Source/Runtime/Engine/Private/Animation/AnimNode_TransitionResult.cpp b/Engine/Source/Runtime/Engine/Private/Animation/AnimNode_TransitionResult.cpp
--- a/Engine/Source/Runtime/Engine/Private/Animation/AnimNode_TransitionResult.cpp
+++ b/Engine/Source/Runtime/Engine/Private/Animation/AnimNode_TransitionResult.cpp
@@ -12,6 +12,7 @@ FAnimNode_TransitionResult::FAnimNode_TransitionResult()
 
 void FAnimNode_TransitionResult::Initialize(const FAnimationInitializeContext& Context)
 {
+	FAnimNode_Base::Initialize(Context);
 }
 
 void FAnimNode_TransitionResult::CacheBones(const FAnimationCacheBonesContext& Context) 
diff --git a/Engine/Source/Runtime/Engine/Private/Animation/AnimNode_UseCachedPose.cpp b/Engine/Source/Runtime/Engine/Private/Animation/AnimNode_UseCachedPose.cpp
--- a/Engine/Source/Runtime/Engine/Private/Animation/AnimNode_UseCachedPose.cpp
+++ b/Engine/Source/Runtime/Engine/Private/Animation/AnimNode_UseCachedPose.cpp
@@ -12,6 +12,8 @@ FAnimNode_UseCachedPose::FAnimNode_UseCachedPose()
 
 void FAnimNode_UseCachedPose::Initialize(const FAnimationInitializeContext& Context)
 {
+	FAnimNode_Base::Initialize(Context);
+
 	LinkToCachingNode.Initialize(Context);
 }
 
diff --git a/Engine/Source/Runtime/Engine/Private/GameEngine.cpp b/Engine/Source/Runtime/Engine/Private/GameEngine.cpp
--- a/Engine/Source/Runtime/Engine/Private/GameEngine.cpp
+++ b/Engine/Source/Runtime/Engine/Private/GameEngine.cpp
@@ -422,6 +422,7 @@ UEngine::UEngine(const FObjectInitializer& ObjectInitializer)
 	EndStreamingPauseDelegate = NULL;
 
 	bCanBlueprintsTickByDefault = true;
+	bOptimizeAnimBlueprintMemberVariableAccess = true;
 
 	bUseFixedFrameRate = false;
 	FixedFrameRate = 30.f;
