CPF_ContainsInstancedReference in a generated UProperty
- Fix for instanced array properties getting wiped out on recompile

diff --git a/Engine/Source/Editor/KismetCompiler/Private/KismetCompilerMisc.cpp b/Engine/Source/Editor/KismetCompiler/Private/KismetCompilerMisc.cpp
--- a/Engine/Source/Editor/KismetCompiler/Private/KismetCompilerMisc.cpp
+++ b/Engine/Source/Editor/KismetCompiler/Private/KismetCompilerMisc.cpp
@@ -807,6 +807,15 @@ UProperty* FKismetCompilerUtilities::CreatePropertyOnScope(UStruct* Scope, const
 				UStructProperty* NewPropertyStruct = NewObject<UStructProperty>(PropertyScope, ValidatedPropertyName, ObjectFlags);
 				NewPropertyStruct->Struct = SubType;
 				NewProperty = NewPropertyStruct;
+
+				if (SubType->StructFlags & STRUCT_HasInstancedReference)
+				{
+					NewProperty->SetPropertyFlags(CPF_ContainsInstancedReference);
+					if (NewArrayProperty)
+					{
+						NewArrayProperty->SetPropertyFlags(CPF_ContainsInstancedReference);
+					}
+				}
 			}
 			else
 			{
diff --git a/Engine/Source/Editor/KismetCompiler/Private/UserDefinedStructureCompilerUtils.cpp b/Engine/Source/Editor/KismetCompiler/Private/UserDefinedStructureCompilerUtils.cpp
--- a/Engine/Source/Editor/KismetCompiler/Private/UserDefinedStructureCompilerUtils.cpp
+++ b/Engine/Source/Editor/KismetCompiler/Private/UserDefinedStructureCompilerUtils.cpp
@@ -205,6 +205,11 @@ struct FUserDefinedStructureCompilerInner
 			VarDesc.CurrentDefaultValue = VarDesc.DefaultValue;
 
 			VarDesc.bInvalidMember = false;
+
+			if (NewProperty->HasAnyPropertyFlags(CPF_InstancedReference | CPF_ContainsInstancedReference))
+			{
+				Struct->StructFlags = EStructFlags(Struct->StructFlags | STRUCT_HasInstancedReference);
+			}
 		}
 	}
 
