[NO NON-PII COMMIT MESSAGE]

diff --git a/Engine/Source/Developer/BlueprintCompilerCppBackend/Private/BlueprintCompilerCppBackendUtils.cpp b/Engine/Source/Developer/BlueprintCompilerCppBackend/Private/BlueprintCompilerCppBackendUtils.cpp
--- a/Engine/Source/Developer/BlueprintCompilerCppBackend/Private/BlueprintCompilerCppBackendUtils.cpp
+++ b/Engine/Source/Developer/BlueprintCompilerCppBackend/Private/BlueprintCompilerCppBackendUtils.cpp
@@ -868,6 +868,13 @@ FString FEmitHelper::LiteralTerm(FEmitterLocalContext& EmitterContext, const FEd
 			return FString(TEXT("nullptr"));
 		}
 	}
+	else if (UEdGraphSchema_K2::PC_Interface == Type.PinCategory)
+	{
+		if (!LiteralObject && CustomValue.IsEmpty())
+		{
+			return FString(TEXT("nullptr"));
+		}
+	}
 	/*
 	else if (CoerceProperty->IsA(UInterfaceProperty::StaticClass()))
 	{
diff --git a/Engine/Source/Editor/KismetCompiler/Private/KismetCompilerMisc.cpp b/Engine/Source/Editor/KismetCompiler/Private/KismetCompilerMisc.cpp
--- a/Engine/Source/Editor/KismetCompiler/Private/KismetCompilerMisc.cpp
+++ b/Engine/Source/Editor/KismetCompiler/Private/KismetCompilerMisc.cpp
@@ -763,7 +763,10 @@ UProperty* FKismetCompilerUtilities::CreatePropertyOnScope(UStruct* Scope, const
 
 		if (SubType != NULL)
 		{
-			if (SubType->HasAnyClassFlags(CLASS_Interface))
+			const bool bIsInterface = SubType->HasAnyClassFlags(CLASS_Interface)
+				|| ((SubType == SelfClass) && ensure(SelfClass->ClassGeneratedBy) && FBlueprintEditorUtils::IsInterfaceBlueprint(CastChecked<UBlueprint>(SelfClass->ClassGeneratedBy)));
+
+			if (bIsInterface)
 			{
 				UInterfaceProperty* NewPropertyObj = NewObject<UInterfaceProperty>(PropertyScope, ValidatedPropertyName, ObjectFlags);
 				// we want to use this setter function instead of setting the 
