Fixing issue where Blueprint func libraries that were expecting a WorldContextObject weren't get implicitly passed one (like they used to be) - removing naive CallableWithoutWorldContext line that we added for Orion (original fix would have been setting ShowWorldContextPin for the class).

diff --git a/Engine/Source/Editor/KismetCompiler/Private/KismetCompiler.cpp b/Engine/Source/Editor/KismetCompiler/Private/KismetCompiler.cpp
--- a/Engine/Source/Editor/KismetCompiler/Private/KismetCompiler.cpp
+++ b/Engine/Source/Editor/KismetCompiler/Private/KismetCompiler.cpp
@@ -1652,8 +1652,6 @@ void FKismetCompilerContext::FinishCompilingFunction(FKismetFunctionContext& Con
 	if (auto WorldContextPin = EntryNode->GetAutoWorldContextPin())
 	{
 		Function->SetMetaData(FBlueprintMetadata::MD_WorldContext, *WorldContextPin->PinName);
-		//should allow function libraries that reference world contexts from objects that don't have one
-		Function->SetMetaData(FBlueprintMetadata::MD_CallableWithoutWorldContext, TEXT("true"));
 	}
 
 	for (int32 EntryPinIndex = 0; EntryPinIndex < EntryNode->Pins.Num(); ++EntryPinIndex)
