Copying //UE4/Dev-Framework to Dev-Main (//UE4/Dev-Main)

diff --git a/Engine/Source/Editor/KismetCompiler/Private/AnimBlueprintCompiler.cpp b/Engine/Source/Editor/KismetCompiler/Private/AnimBlueprintCompiler.cpp
--- a/Engine/Source/Editor/KismetCompiler/Private/AnimBlueprintCompiler.cpp
+++ b/Engine/Source/Editor/KismetCompiler/Private/AnimBlueprintCompiler.cpp
@@ -1833,8 +1833,6 @@ void FAnimBlueprintCompiler::FEvaluationHandlerRecord::PatchFunctionNameAndCopyR
 		const FAnimNodeSinglePropertyHandler& AnimNodeSinglePropertyHandler = PropertyIt.Value();
 		if(AnimNodeSinglePropertyHandler.SimpleCopyPropertyName != NAME_None)
 		{
-			// Local variable get, no need to thunk to BP to grab this value
-
 			UProperty* SimpleCopyPropertySource = TargetObject->GetClass()->FindPropertyByName(AnimNodeSinglePropertyHandler.SimpleCopyPropertyName);
 			if (SimpleCopyPropertySource)
 			{
@@ -1848,25 +1846,51 @@ void FAnimBlueprintCompiler::FEvaluationHandlerRecord::PatchFunctionNameAndCopyR
 						FExposedValueCopyRecord CopyRecord;
 						CopyRecord.DestProperty = SimpleCopyPropertyDest;
 						CopyRecord.DestArrayIndex = ArrayIt.Key();
-						CopyRecord.SourceProperty = SimpleCopyPropertySource;
+						CopyRecord.SourcePropertyName = AnimNodeSinglePropertyHandler.SimpleCopyPropertyName;
+						CopyRecord.SourceSubPropertyName = NAME_None;
 						CopyRecord.SourceArrayIndex = 0;
 						CopyRecord.Size = SimpleCopyPropertyDest->Inner->GetSize();
+						CopyRecord.PostCopyOperation = AnimNodeSinglePropertyHandler.Operation;
 						HandlerPtr->CopyRecords.Add(CopyRecord);
 					}
 				}
 				else
 				{
 					UProperty* SimpleCopyPropertyDest = NodeVariableProperty->Struct->FindPropertyByName(PropertyIt.Key());
-					check(SimpleCopyPropertySource->GetSize() == SimpleCopyPropertyDest->GetSize());
-
-					// Local variable get, no need to thunk to BP to grab this value
-					FExposedValueCopyRecord CopyRecord;
-					CopyRecord.DestProperty = SimpleCopyPropertyDest;
-					CopyRecord.DestArrayIndex = 0;
-					CopyRecord.SourceProperty = SimpleCopyPropertySource;
-					CopyRecord.SourceArrayIndex = 0;
-					CopyRecord.Size = SimpleCopyPropertyDest->GetSize();
-					HandlerPtr->CopyRecords.Add(CopyRecord);
+					check(SimpleCopyPropertyDest);
+					if(AnimNodeSinglePropertyHandler.SubStructPropertyName != NAME_None)
+					{
+						UStructProperty* SourceStructProperty = CastChecked<UStructProperty>(SimpleCopyPropertySource);
+						UProperty* SourceSubProperty = SourceStructProperty->Struct->FindPropertyByName(AnimNodeSinglePropertyHandler.SubStructPropertyName);
+						check(SourceSubProperty);
+						check(SourceSubProperty->GetSize() == SimpleCopyPropertyDest->GetSize());
+
+						// Local sub-struct variable get
+						FExposedValueCopyRecord CopyRecord;
+						CopyRecord.DestProperty = SimpleCopyPropertyDest;
+						CopyRecord.DestArrayIndex = 0;
+						CopyRecord.SourcePropertyName = AnimNodeSinglePropertyHandler.SimpleCopyPropertyName;
+						CopyRecord.SourceSubPropertyName = AnimNodeSinglePropertyHandler.SubStructPropertyName;
+						CopyRecord.SourceArrayIndex = 0;
+						CopyRecord.Size = SimpleCopyPropertyDest->GetSize();
+						CopyRecord.PostCopyOperation = AnimNodeSinglePropertyHandler.Operation;
+						HandlerPtr->CopyRecords.Add(CopyRecord);
+					}
+					else
+					{
+						check(SimpleCopyPropertySource->GetSize() == SimpleCopyPropertyDest->GetSize());
+
+						// Local variable get
+						FExposedValueCopyRecord CopyRecord;
+						CopyRecord.DestProperty = SimpleCopyPropertyDest;
+						CopyRecord.DestArrayIndex = 0;
+						CopyRecord.SourcePropertyName = AnimNodeSinglePropertyHandler.SimpleCopyPropertyName;
+						CopyRecord.SourceSubPropertyName = NAME_None;
+						CopyRecord.SourceArrayIndex = 0;
+						CopyRecord.Size = SimpleCopyPropertyDest->GetSize();
+						CopyRecord.PostCopyOperation = AnimNodeSinglePropertyHandler.Operation;
+						HandlerPtr->CopyRecords.Add(CopyRecord);
+					}
 				}
 			}
 			else
@@ -1888,6 +1912,47 @@ void FAnimBlueprintCompiler::FEvaluationHandlerRecord::PatchFunctionNameAndCopyR
 	}
 }
 
+static UEdGraphPin* FindFirstInputPin(UEdGraphNode* InNode)
+{
+	const UAnimationGraphSchema* Schema = GetDefault<UAnimationGraphSchema>();
+
+	for(UEdGraphPin* Pin : InNode->Pins)
+	{
+		if(Pin && Pin->Direction == EGPD_Input && !Schema->IsExecPin(*Pin) && !Schema->IsSelfPin(*Pin))
+		{
+			return Pin;
+		}
+	}
+
+	return nullptr;
+}
+
+static UEdGraphNode* FollowKnots(UEdGraphPin* FromPin, UEdGraphPin*& DestPin)
+{
+	UEdGraphPin* LinkedPin = FromPin->LinkedTo[0];
+	DestPin = LinkedPin;
+	if(LinkedPin)
+	{
+		UEdGraphNode* LinkedNode = LinkedPin->GetOwningNode();
+		UK2Node_Knot* KnotNode = Cast<UK2Node_Knot>(LinkedNode);
+		while(KnotNode)
+		{
+			if(UEdGraphPin* InputPin = FindFirstInputPin(KnotNode))
+			{
+				if(InputPin->LinkedTo[0])
+				{
+					DestPin = InputPin->LinkedTo[0];
+					LinkedNode = InputPin->LinkedTo[0]->GetOwningNode();
+					KnotNode = Cast<UK2Node_Knot>(LinkedNode);
+				}
+			}
+		}
+		return LinkedNode;
+	}
+
+	return nullptr;
+}
+
 void FAnimBlueprintCompiler::FEvaluationHandlerRecord::RegisterPin(UEdGraphPin* SourcePin, UProperty* AssociatedProperty, int32 AssociatedPropertyArrayIndex)
 {
 	const UAnimationGraphSchema* Schema = GetDefault<UAnimationGraphSchema>();
@@ -1906,64 +1971,189 @@ void FAnimBlueprintCompiler::FEvaluationHandlerRecord::RegisterPin(UEdGraphPin*
 
 	if(GetDefault<UEngine>()->bOptimizeAnimBlueprintMemberVariableAccess)
 	{
-		UK2Node_VariableGet* VariableGetNode = Cast<UK2Node_VariableGet>(SourcePin->LinkedTo[0]->GetOwningNode());
-		if(VariableGetNode && VariableGetNode->IsNodePure() && VariableGetNode->VariableReference.IsSelfContext())
+		typedef bool (FAnimBlueprintCompiler::FEvaluationHandlerRecord::*GraphCheckerFunc)(FAnimNodeSinglePropertyHandler&, UEdGraphPin*);
+
+		GraphCheckerFunc GraphCheckerFuncs[] =
+		{
+			&FAnimBlueprintCompiler::FEvaluationHandlerRecord::CheckForVariableGet,
+			&FAnimBlueprintCompiler::FEvaluationHandlerRecord::CheckForLogicalNot,
+			&FAnimBlueprintCompiler::FEvaluationHandlerRecord::CheckForStructMemberAccess,
+		};
+
+		for(GraphCheckerFunc& CheckFunc : GraphCheckerFuncs)
 		{
-			Handler.SimpleCopyPropertyName = VariableGetNode->VariableReference.GetMemberName();
+			if((this->*CheckFunc)(Handler, SourcePin))
+			{
+				break;
+			}
 		}
-		else
+
+		CheckForMemberOnlyAccess(Handler, SourcePin);
+	}
+}
+
+static FName RecoverSplitStructPinName(UEdGraphPin* OutputPin)
+{
+	check(OutputPin->ParentPin);
+	
+	const FString& PinName = OutputPin->PinName;
+	const FString& ParentPinName = OutputPin->ParentPin->PinName;
+
+	return FName(*PinName.Replace(*(ParentPinName + TEXT("_")), TEXT("")));
+}
+
+bool FAnimBlueprintCompiler::FEvaluationHandlerRecord::CheckForVariableGet(FAnimNodeSinglePropertyHandler& Handler, UEdGraphPin* SourcePin)
+{
+	if(SourcePin)
+	{
+		UEdGraphPin* DestPin = nullptr;
+		if(UK2Node_VariableGet* VariableGetNode = Cast<UK2Node_VariableGet>(FollowKnots(SourcePin, DestPin)))
+		{
+			if(VariableGetNode && VariableGetNode->IsNodePure() && VariableGetNode->VariableReference.IsSelfContext())
+			{
+				if(DestPin)
+				{
+					// variable get could be a 'split' struct
+					if(DestPin->ParentPin != nullptr)
+					{
+						Handler.SimpleCopyPropertyName = VariableGetNode->VariableReference.GetMemberName();
+						Handler.SubStructPropertyName = RecoverSplitStructPinName(DestPin);
+					}
+					else
+					{
+						Handler.SimpleCopyPropertyName = VariableGetNode->VariableReference.GetMemberName();
+					}
+					return true;
+				}
+			}
+		}
+	}
+
+	return false;
+}
+
+bool FAnimBlueprintCompiler::FEvaluationHandlerRecord::CheckForLogicalNot(FAnimNodeSinglePropertyHandler& Handler, UEdGraphPin* SourcePin)
+{
+	if(SourcePin)
+	{
+		UEdGraphPin* DestPin = nullptr;
+		UK2Node_CallFunction* CallFunctionNode = Cast<UK2Node_CallFunction>(FollowKnots(SourcePin, DestPin));
+		if(CallFunctionNode && CallFunctionNode->FunctionReference.GetMemberName() == FName(TEXT("Not_PreBool")))
 		{
-			Handler.bHasOnlyMemberAccess = true;
+			// find and follow input pin
+			if(UEdGraphPin* InputPin = FindFirstInputPin(CallFunctionNode))
+			{
+				check(InputPin->PinType.PinCategory == UEdGraphSchema_K2::PC_Boolean);
+				if(CheckForVariableGet(Handler, InputPin) || CheckForStructMemberAccess(Handler, InputPin))
+				{
+					check(Handler.SimpleCopyPropertyName != NAME_None);	// this should have been filled in by CheckForVariableGet() or CheckForStructMemberAccess() above
+					Handler.Operation = EPostCopyOperation::LogicalNegateBool;
+					return true;
+				}
+			}
+		}
+	}
 
-			// traverse pins to leaf nodes and check for member access/pure only
-			TArray<UEdGraphPin*> PinStack;
-			PinStack.Add(SourcePin);
-			while(PinStack.Num() > 0)
+	return false;
+}
+
+bool FAnimBlueprintCompiler::FEvaluationHandlerRecord::CheckForStructMemberAccess(FAnimNodeSinglePropertyHandler& Handler, UEdGraphPin* SourcePin)
+{
+	if(SourcePin)
+	{
+		UEdGraphPin* DestPin = nullptr;
+		if(UK2Node_BreakStruct* BreakStructNode = Cast<UK2Node_BreakStruct>(FollowKnots(SourcePin, DestPin)))
+		{
+			if(UEdGraphPin* InputPin = FindFirstInputPin(BreakStructNode))
+			{
+				if(CheckForVariableGet(Handler, InputPin))
+				{
+					check(Handler.SimpleCopyPropertyName != NAME_None);	// this should have been filled in by CheckForVariableGet() above
+					Handler.SubStructPropertyName = *DestPin->PinName;
+					return true;
+				}
+			}
+		}
+		// could be a native break
+		else if(UK2Node_CallFunction* NativeBreakNode = Cast<UK2Node_CallFunction>(FollowKnots(SourcePin, DestPin)))
+		{
+			UFunction* Function = NativeBreakNode->FunctionReference.ResolveMember<UFunction>(UKismetMathLibrary::StaticClass());
+			if(Function && Function->HasMetaData(TEXT("NativeBreakFunc")))
 			{
-				UEdGraphPin* CurrentPin = PinStack.Pop(false);
-				for(auto& LinkedPin : CurrentPin->LinkedTo)
+				if(UEdGraphPin* InputPin = FindFirstInputPin(NativeBreakNode))
 				{
-					UEdGraphNode* LinkedNode = LinkedPin->GetOwningNode();
-					if(LinkedNode)
+					if(CheckForVariableGet(Handler, InputPin))
 					{
-						bool bLeafNode = true;
-						for(auto& Pin : LinkedNode->Pins)
+						check(Handler.SimpleCopyPropertyName != NAME_None);	// this should have been filled in by CheckForVariableGet() above
+						Handler.SubStructPropertyName = *DestPin->PinName;
+						return true;
+					}
+				}
+			}
+		}
+	}
+
+	return false;
+}
+
+bool FAnimBlueprintCompiler::FEvaluationHandlerRecord::CheckForMemberOnlyAccess(FAnimNodeSinglePropertyHandler& Handler, UEdGraphPin* SourcePin)
+{
+	const UAnimationGraphSchema* Schema = GetDefault<UAnimationGraphSchema>();
+
+	Handler.bHasOnlyMemberAccess = true;
+
+	if(SourcePin)
+	{
+		// traverse pins to leaf nodes and check for member access/pure only
+		TArray<UEdGraphPin*> PinStack;
+		PinStack.Add(SourcePin);
+		while(PinStack.Num() > 0)
+		{
+			UEdGraphPin* CurrentPin = PinStack.Pop(false);
+			for(auto& LinkedPin : CurrentPin->LinkedTo)
+			{
+				UEdGraphNode* LinkedNode = LinkedPin->GetOwningNode();
+				if(LinkedNode)
+				{
+					bool bLeafNode = true;
+					for(auto& Pin : LinkedNode->Pins)
+					{
+						if(Pin != LinkedPin && Pin->Direction == EGPD_Input && !Schema->IsPosePin(Pin->PinType))
 						{
-							if(Pin != LinkedPin && Pin->Direction == EGPD_Input && !Schema->IsPosePin(Pin->PinType))
-							{
-								bLeafNode = false;
-								PinStack.Add(Pin);
-							}
+							bLeafNode = false;
+							PinStack.Add(Pin);
 						}
+					}
 
-						if(bLeafNode)
+					if(bLeafNode)
+					{
+						if(UK2Node_VariableGet* LinkedVariableGetNode = Cast<UK2Node_VariableGet>(LinkedNode))
 						{
-							if(UK2Node_VariableGet* LinkedVariableGetNode = Cast<UK2Node_VariableGet>(LinkedNode))
-							{
-								if(!LinkedVariableGetNode->IsNodePure() || !LinkedVariableGetNode->VariableReference.IsSelfContext())
-								{
-									// only local variable access is allowed for leaf nodes, so 
-									Handler.bHasOnlyMemberAccess = false;
-								}
-							}
-							else if(UK2Node_CallFunction* CallFunctionNode = Cast<UK2Node_CallFunction>(LinkedNode))
+							if(!LinkedVariableGetNode->IsNodePure() || !LinkedVariableGetNode->VariableReference.IsSelfContext())
 							{
-								if(!CallFunctionNode->IsNodePure())
-								{
-									// only allow pure function calls
-									Handler.bHasOnlyMemberAccess = false;
-								}
+								// only local variable access is allowed for leaf nodes 
+								Handler.bHasOnlyMemberAccess = false;
 							}
-							else if(!LinkedNode->IsA<UK2Node_TransitionRuleGetter>())
+						}
+						else if(UK2Node_CallFunction* CallFunctionNode = Cast<UK2Node_CallFunction>(LinkedNode))
+						{
+							if(!CallFunctionNode->IsNodePure())
 							{
+								// only allow pure function calls
 								Handler.bHasOnlyMemberAccess = false;
 							}
 						}
+						else if(!LinkedNode->IsA<UK2Node_TransitionRuleGetter>())
+						{
+							Handler.bHasOnlyMemberAccess = false;
+						}
 					}
 				}
 			}
 		}
 	}
+
+	return Handler.bHasOnlyMemberAccess;
 }
 
 //////////////////////////////////////////////////////////////////////////
diff --git a/Engine/Source/Editor/KismetCompiler/Private/AnimBlueprintCompiler.h b/Engine/Source/Editor/KismetCompiler/Private/AnimBlueprintCompiler.h
--- a/Engine/Source/Editor/KismetCompiler/Private/AnimBlueprintCompiler.h
+++ b/Engine/Source/Editor/KismetCompiler/Private/AnimBlueprintCompiler.h
@@ -53,15 +53,23 @@ protected:
 
 		UEdGraphPin* SinglePin;
 
-		// If this is valid we are a simple property copy, avoiding BP thunk
+		// If this is valid we are a 'simple property copy', avoiding BP thunk
 		FName SimpleCopyPropertyName;
 
+		// If this is a sub-struct property, this will be something other than NAME_None
+		FName SubStructPropertyName;
+
+		// Any operation we want to perform post-copy on the destination data
+		EPostCopyOperation Operation;
+
 		// Whether this node chain only ever accesses members
 		bool bHasOnlyMemberAccess;
 
 		FAnimNodeSinglePropertyHandler()
 			: SinglePin(NULL)
 			, SimpleCopyPropertyName(NAME_None)
+			, SubStructPropertyName(NAME_None)
+			, Operation(EPostCopyOperation::None)
 			, bHasOnlyMemberAccess(false)
 		{
 		}
@@ -147,6 +155,15 @@ protected:
 		void PatchFunctionNameAndCopyRecordsInto(UObject* TargetObject) const;
 
 		void RegisterPin(UEdGraphPin* SourcePin, UProperty* AssociatedProperty, int32 AssociatedPropertyArrayIndex);
+
+	private:
+		bool CheckForVariableGet(FAnimNodeSinglePropertyHandler& Handler, UEdGraphPin* SourcePin);
+
+		bool CheckForLogicalNot(FAnimNodeSinglePropertyHandler& Handler, UEdGraphPin* SourcePin);
+
+		bool CheckForStructMemberAccess(FAnimNodeSinglePropertyHandler& Handler, UEdGraphPin* SourcePin);
+
+		bool CheckForMemberOnlyAccess(FAnimNodeSinglePropertyHandler& Handler, UEdGraphPin* SourcePin);
 	};
 
 	// State machines may get processed before their inner graphs, so their node index needs to be patched up later
diff --git a/Engine/Source/Editor/MergeActors/Private/MeshMergingTool/SMeshMergingDialog.cpp b/Engine/Source/Editor/MergeActors/Private/MeshMergingTool/SMeshMergingDialog.cpp
--- a/Engine/Source/Editor/MergeActors/Private/MeshMergingTool/SMeshMergingDialog.cpp
+++ b/Engine/Source/Editor/MergeActors/Private/MeshMergingTool/SMeshMergingDialog.cpp
@@ -513,6 +513,7 @@ void SMeshMergingDialog::SetExportSpecularMap(ECheckBoxState NewValue)
 void SMeshMergingDialog::SetMergedMaterialAtlasResolution(TSharedPtr<FString> NewSelection, ESelectInfo::Type SelectInfo)
 {
 	TTypeFromString<int32>::FromString(Tool->MergingSettings.MaterialSettings.TextureSize.X, **NewSelection);
+	TTypeFromString<int32>::FromString(Tool->MergingSettings.MaterialSettings.TextureSize.Y, **NewSelection);
 }
 
 
diff --git a/Engine/Source/Editor/Persona/Private/AnimationEditorViewportClient.cpp b/Engine/Source/Editor/Persona/Private/AnimationEditorViewportClient.cpp
--- a/Engine/Source/Editor/Persona/Private/AnimationEditorViewportClient.cpp
+++ b/Engine/Source/Editor/Persona/Private/AnimationEditorViewportClient.cpp
@@ -911,7 +911,7 @@ void FAnimationViewportClient::DisplayInfo(FCanvas* Canvas, FSceneView* View, bo
 	if( PreviewInstance )
 	{
 		// see if you have anim sequence that has transform curves
-		UAnimSequence* Sequence = Cast<UAnimSequence>(PreviewInstance->CurrentAsset);
+		UAnimSequence* Sequence = Cast<UAnimSequence>(PreviewInstance->GetCurrentAsset());
 		if ( Sequence && Sequence->DoesNeedRebake() )
 		{
 			InfoString = TEXT("Animation is being edited. To apply to raw animation data, click \"Apply\"");
diff --git a/Engine/Source/Editor/Persona/Private/Persona.cpp b/Engine/Source/Editor/Persona/Private/Persona.cpp
--- a/Engine/Source/Editor/Persona/Private/Persona.cpp
+++ b/Engine/Source/Editor/Persona/Private/Persona.cpp
@@ -1037,7 +1037,7 @@ UObject* FPersona::GetPreviewAnimationAsset() const
 			// if same, do not overwrite. It will reset time and everything
 			if (PreviewComponent->PreviewInstance != NULL)
 			{
-				return PreviewComponent->PreviewInstance->CurrentAsset;
+				return PreviewComponent->PreviewInstance->GetCurrentAsset();
 			}
 		}
 	}
@@ -1094,7 +1094,7 @@ void FPersona::SetPreviewVertexAnim(UVertexAnimation* VertexAnim)
 			// if same, do not overwrite. It will reset time and everything
 			if( VertexAnim && 
 				PreviewComponent->PreviewInstance && 
-				VertexAnim == PreviewComponent->PreviewInstance->CurrentVertexAnim )
+				VertexAnim == PreviewComponent->PreviewInstance->GetCurrentVertexAnimation() )
 			{
 				return;
 			}
@@ -1884,9 +1884,9 @@ void FPersona::SetPreviewMeshInternal(USkeletalMesh* NewPreviewMesh)
 		bool bNeedsToCopyAnimationData = PreviewComponent->GetAnimInstance() && PreviewComponent->GetAnimInstance() == PreviewComponent->PreviewInstance;
 		if(bNeedsToCopyAnimationData)
 		{
-			AnimAssetToPlay = PreviewComponent->PreviewInstance->CurrentAsset;
-			PlayPosition = PreviewComponent->PreviewInstance->CurrentTime;
-			bPlaying = PreviewComponent->PreviewInstance->bPlaying;
+			AnimAssetToPlay = PreviewComponent->PreviewInstance->GetCurrentAsset();
+			PlayPosition = PreviewComponent->PreviewInstance->GetCurrentTime();
+			bPlaying = PreviewComponent->PreviewInstance->IsPlaying();
 		}
 
 		PreviewComponent->SetSkeletalMesh(NewPreviewMesh);
@@ -1895,7 +1895,7 @@ void FPersona::SetPreviewMeshInternal(USkeletalMesh* NewPreviewMesh)
 		{
 			SetPreviewAnimationAsset(AnimAssetToPlay);
 			PreviewComponent->PreviewInstance->SetPosition(PlayPosition);
-			PreviewComponent->PreviewInstance->bPlaying = bPlaying;
+			PreviewComponent->PreviewInstance->SetPlaying(bPlaying);
 		}
 	}
 	else
@@ -1981,10 +1981,14 @@ void FPersona::RefreshPreviewInstanceTrackCurves()
 
 void FPersona::PostUndo(bool bSuccess)
 {
+	DocumentManager->CleanInvalidTabs();
 	DocumentManager->RefreshAllTabs();
 
 	FBlueprintEditor::PostUndo(bSuccess);
 
+	// If we undid a node creation that caused us to clean up a tab/graph we need to refresh the UI state
+	RefreshEditors();
+
 	// PostUndo broadcast
 	OnPostUndo.Broadcast();	
 
@@ -3254,9 +3258,9 @@ void FPersona::ShowReferencePose(bool bReferencePose)
 		}
 		else
 		{
-			if (PreviewComponent->PreviewInstance && PreviewComponent->PreviewInstance->CurrentAsset)
+			if (PreviewComponent->PreviewInstance && PreviewComponent->PreviewInstance->GetCurrentAsset())
 			{
-				CachedPreviewAsset = PreviewComponent->PreviewInstance->CurrentAsset;
+				CachedPreviewAsset = PreviewComponent->PreviewInstance->GetCurrentAsset();
 			}
 			
 			PreviewComponent->EnablePreview(true, NULL, NULL);
@@ -3273,7 +3277,7 @@ bool FPersona::IsShowReferencePoseEnabled() const
 {
 	if(PreviewComponent)
 	{
-		return PreviewComponent->IsPreviewOn() && PreviewComponent->PreviewInstance->CurrentAsset == NULL;
+		return PreviewComponent->IsPreviewOn() && PreviewComponent->PreviewInstance->GetCurrentAsset() == NULL;
 	}
 	return false;
 }
diff --git a/Engine/Source/Editor/Persona/Private/SAnimEditorBase.cpp b/Engine/Source/Editor/Persona/Private/SAnimEditorBase.cpp
--- a/Engine/Source/Editor/Persona/Private/SAnimEditorBase.cpp
+++ b/Engine/Source/Editor/Persona/Private/SAnimEditorBase.cpp
@@ -247,7 +247,7 @@ void SAnimEditorBase::RecalculateSequenceLength()
 			if (PreviewInstance)
 			{
 				// Re-set the position, so instance is clamped properly
-				PreviewInstance->SetPosition(PreviewInstance->CurrentTime, false); 
+				PreviewInstance->SetPosition(PreviewInstance->GetCurrentTime(), false); 
 			}
 		}
 	}
@@ -292,7 +292,7 @@ float SAnimEditorBase::GetScrubValue() const
 	UAnimSingleNodeInstance * PreviewInstance = GetPreviewInstance();
 	if (PreviewInstance)
 	{
-		float CurTime = PreviewInstance->CurrentTime;
+		float CurTime = PreviewInstance->GetCurrentTime();
 		return (CurTime); 
 	}
 	else
@@ -315,7 +315,7 @@ FText SAnimEditorBase::GetCurrentSequenceTime() const
 
 	if (PreviewInstance)
 	{
-		CurTime = PreviewInstance->CurrentTime;
+		CurTime = PreviewInstance->GetCurrentTime();
 	}
 
 	static const FNumberFormattingOptions FractionNumberFormat = FNumberFormattingOptions()
@@ -330,7 +330,7 @@ FText SAnimEditorBase::GetCurrentPercentage() const
 	float Percentage = 0.f;
 	if (PreviewInstance)
 	{
-		Percentage = PreviewInstance->CurrentTime / GetEditorObject()->SequenceLength;
+		Percentage = PreviewInstance->GetCurrentTime() / GetEditorObject()->SequenceLength;
 	}
 
 	static const FNumberFormattingOptions PercentNumberFormat = FNumberFormattingOptions()
@@ -347,7 +347,7 @@ FText SAnimEditorBase::GetCurrentFrame() const
 
 	if (PreviewInstance)
 	{
-		Percentage = PreviewInstance->CurrentTime/GetEditorObject()->SequenceLength;
+		Percentage = PreviewInstance->GetCurrentTime()/GetEditorObject()->SequenceLength;
 	}
 
 	static const FNumberFormattingOptions FractionNumberFormat = FNumberFormattingOptions()
diff --git a/Engine/Source/Editor/Persona/Private/SAnimMontageScrubPanel.cpp b/Engine/Source/Editor/Persona/Private/SAnimMontageScrubPanel.cpp
--- a/Engine/Source/Editor/Persona/Private/SAnimMontageScrubPanel.cpp
+++ b/Engine/Source/Editor/Persona/Private/SAnimMontageScrubPanel.cpp
@@ -26,7 +26,7 @@ FReply SAnimMontageScrubPanel::OnClick_ToggleLoop()
 {
 	if (UAnimPreviewInstance* AnimPreviewInstance = Cast<UAnimPreviewInstance>(GetPreviewInstance()))
 	{
-		bool bIsLooping = AnimPreviewInstance->bLooping;
+		bool bIsLooping = AnimPreviewInstance->IsLooping();
 		AnimPreviewInstance->MontagePreview_SetLooping(!bIsLooping);
 	}
 	return FReply::Handled();
@@ -36,7 +36,7 @@ FReply SAnimMontageScrubPanel::OnClick_Backward()
 {
 	if (UAnimPreviewInstance* AnimPreviewInstance = Cast<UAnimPreviewInstance>(GetPreviewInstance()))
 	{
-		if (! AnimPreviewInstance->bPlaying || ! AnimPreviewInstance->bReverse)
+		if (! AnimPreviewInstance->IsPlaying() || ! AnimPreviewInstance->IsReverse())
 		{
 			// check if there is montage being played, if not, start it, if there is, just unpause
 			if (! AnimPreviewInstance->IsPlayingMontage())
@@ -64,7 +64,7 @@ FReply SAnimMontageScrubPanel::OnClick_Forward()
 {
 	if (UAnimPreviewInstance* AnimPreviewInstance = Cast<UAnimPreviewInstance>(GetPreviewInstance()))
 	{
-		if (! AnimPreviewInstance->bPlaying || AnimPreviewInstance->bReverse)
+		if (! AnimPreviewInstance->IsPlaying() || AnimPreviewInstance->IsReverse())
 		{
 			// check if there is montage being played, if not, start it, if there is, just unpause
 			if (! AnimPreviewInstance->IsPlayingMontage())
diff --git a/Engine/Source/Editor/Persona/Private/SAnimationBlendSpaceBase.cpp b/Engine/Source/Editor/Persona/Private/SAnimationBlendSpaceBase.cpp
--- a/Engine/Source/Editor/Persona/Private/SAnimationBlendSpaceBase.cpp
+++ b/Engine/Source/Editor/Persona/Private/SAnimationBlendSpaceBase.cpp
@@ -1282,7 +1282,7 @@ ECheckBoxState SBlendSpaceEditorBase::IsPreviewOn() const
 
 	if (Component != NULL && Component->IsPreviewOn())
 	{
-		if (Component->PreviewInstance->CurrentAsset == BlendSpace)
+		if (Component->PreviewInstance->GetCurrentAsset() == BlendSpace)
 		{
 			return ECheckBoxState::Checked;
 		}
@@ -1323,7 +1323,7 @@ void SBlendSpaceEditorBase::UpdatePreviewParameter() const
 
 	if (Component != nullptr && Component->IsPreviewOn())
 	{
-		if (Component->PreviewInstance->CurrentAsset == BlendSpace)
+		if (Component->PreviewInstance->GetCurrentAsset() == BlendSpace)
 		{
 			FVector BlendInput = BlendSpaceWidget->LastValidMouseEditorPoint;
 			Component->PreviewInstance->SetBlendSpaceInput(BlendInput);
diff --git a/Engine/Source/Editor/Persona/Private/SAnimationEditorViewport.cpp b/Engine/Source/Editor/Persona/Private/SAnimationEditorViewport.cpp
--- a/Engine/Source/Editor/Persona/Private/SAnimationEditorViewport.cpp
+++ b/Engine/Source/Editor/Persona/Private/SAnimationEditorViewport.cpp
@@ -1111,7 +1111,7 @@ void SAnimationEditorViewportTabBody::ShowRetargetBasePose()
 	UDebugSkelMeshComponent* PreviewComponent = PersonaPtr.Pin()->PreviewComponent;
 	if(PreviewComponent && PreviewComponent->PreviewInstance)
 	{
-		PreviewComponent->PreviewInstance->bForceRetargetBasePose = !PreviewComponent->PreviewInstance->bForceRetargetBasePose;
+		PreviewComponent->PreviewInstance->SetForceRetargetBasePose(!PreviewComponent->PreviewInstance->GetForceRetargetBasePose());
 	}
 }
 
@@ -1126,7 +1126,7 @@ bool SAnimationEditorViewportTabBody::IsShowRetargetBasePoseEnabled() const
 	UDebugSkelMeshComponent* PreviewComponent = PersonaPtr.Pin()->PreviewComponent;
 	if(PreviewComponent && PreviewComponent->PreviewInstance)
 	{
-		return PreviewComponent->PreviewInstance->bForceRetargetBasePose;
+		return PreviewComponent->PreviewInstance->GetForceRetargetBasePose();
 	}
 	return false;
 }
diff --git a/Engine/Source/Editor/Persona/Private/SAnimationScrubPanel.cpp b/Engine/Source/Editor/Persona/Private/SAnimationScrubPanel.cpp
--- a/Engine/Source/Editor/Persona/Private/SAnimationScrubPanel.cpp
+++ b/Engine/Source/Editor/Persona/Private/SAnimationScrubPanel.cpp
@@ -133,8 +133,8 @@ FReply SAnimationScrubPanel::OnClick_Forward()
 	UAnimSingleNodeInstance* PreviewInstance = GetPreviewInstance();
 	if (PreviewInstance)
 	{
-		bool bIsReverse = PreviewInstance->bReverse;
-		bool bIsPlaying = PreviewInstance->bPlaying;
+		bool bIsReverse = PreviewInstance->IsReverse();
+		bool bIsPlaying = PreviewInstance->IsPlaying();
 		// if current bIsReverse and bIsPlaying, we'd like to just turn off reverse
 		if (bIsReverse && bIsPlaying)
 		{
@@ -172,8 +172,8 @@ FReply SAnimationScrubPanel::OnClick_Backward()
 	UAnimSingleNodeInstance* PreviewInstance = GetPreviewInstance();
 	if (PreviewInstance)
 	{
-		bool bIsReverse = PreviewInstance->bReverse;
-		bool bIsPlaying = PreviewInstance->bPlaying;
+		bool bIsReverse = PreviewInstance->IsReverse();
+		bool bIsPlaying = PreviewInstance->IsPlaying();
 		// if currently playing forward, just simply turn on reverse
 		if (!bIsReverse && bIsPlaying)
 		{
@@ -203,7 +203,7 @@ FReply SAnimationScrubPanel::OnClick_ToggleLoop()
 	UAnimSingleNodeInstance* PreviewInstance = GetPreviewInstance();
 	if (PreviewInstance)
 	{
-		bool bIsLooping = PreviewInstance->bLooping;
+		bool bIsLooping = PreviewInstance->IsLooping();
 		PreviewInstance->SetLooping(!bIsLooping);
 	}
 	return FReply::Handled();
@@ -212,16 +212,16 @@ FReply SAnimationScrubPanel::OnClick_ToggleLoop()
 bool SAnimationScrubPanel::IsLoopStatusOn() const
 {
 	UAnimSingleNodeInstance* PreviewInstance = GetPreviewInstance();
-	return (PreviewInstance && PreviewInstance->bLooping);
+	return (PreviewInstance && PreviewInstance->IsLooping());
 }
 
 EPlaybackMode::Type SAnimationScrubPanel::GetPlaybackMode() const
 {
 	if (UAnimSingleNodeInstance* PreviewInstance = GetPreviewInstance())
 	{
-		if (PreviewInstance->bPlaying)
+		if (PreviewInstance->IsPlaying())
 		{
-			return PreviewInstance->bReverse ? EPlaybackMode::PlayingReverse : EPlaybackMode::PlayingForward;
+			return PreviewInstance->IsReverse() ? EPlaybackMode::PlayingReverse : EPlaybackMode::PlayingForward;
 		}
 		return EPlaybackMode::Stopped;
 	}
@@ -236,7 +236,7 @@ EPlaybackMode::Type SAnimationScrubPanel::GetPlaybackMode() const
 bool SAnimationScrubPanel::IsRealtimeStreamingMode() const
 {
 	UAnimSingleNodeInstance* PreviewInstance = GetPreviewInstance();
-	return ( ! (PreviewInstance && PreviewInstance->CurrentAsset) );
+	return ( ! (PreviewInstance && PreviewInstance->GetCurrentAsset()) );
 }
 
 void SAnimationScrubPanel::OnValueChanged(float NewValue)
@@ -284,9 +284,9 @@ uint32 SAnimationScrubPanel::GetNumOfFrames() const
 		float Length = PreviewInstance->GetLength();
 		// if anim sequence, use correct num frames
 		int32 NumFrames = (int32) (Length/0.0333f); 
-		if (PreviewInstance->CurrentAsset && PreviewInstance->CurrentAsset->IsA(UAnimSequenceBase::StaticClass()))
+		if (PreviewInstance->GetCurrentAsset() && PreviewInstance->GetCurrentAsset()->IsA(UAnimSequenceBase::StaticClass()))
 		{
-			NumFrames = CastChecked<UAnimSequenceBase>(PreviewInstance->CurrentAsset)->GetNumberOfFrames();
+			NumFrames = CastChecked<UAnimSequenceBase>(PreviewInstance->GetCurrentAsset())->GetNumberOfFrames();
 		}
 		return NumFrames;
 	}
@@ -335,7 +335,7 @@ bool SAnimationScrubPanel::DoesSyncViewport() const
 {
 	UAnimSingleNodeInstance* PreviewInstance = GetPreviewInstance();
 
-	return (( LockedSequence==NULL && PreviewInstance ) || ( LockedSequence && PreviewInstance && PreviewInstance->CurrentAsset == LockedSequence ));
+	return (( LockedSequence==NULL && PreviewInstance ) || ( LockedSequence && PreviewInstance && PreviewInstance->GetCurrentAsset() == LockedSequence ));
 }
 
 void SAnimationScrubPanel::Tick( const FGeometry& AllottedGeometry, const double InCurrentTime, const float InDeltaTime )
@@ -367,7 +367,7 @@ float SAnimationScrubPanel::GetScrubValue() const
 		UAnimSingleNodeInstance* PreviewInstance = GetPreviewInstance();
 		if (PreviewInstance)
 		{
-			return PreviewInstance->CurrentTime; 
+			return PreviewInstance->GetCurrentTime(); 
 		}
 	}
 	else
@@ -432,9 +432,9 @@ void SAnimationScrubPanel::OnCropAnimSequence( bool bFromStart, float CurrentTim
 	if(PreviewInstance)
 	{
 		float Length = PreviewInstance->GetLength();
-		if (PreviewInstance->CurrentAsset)
+		if (PreviewInstance->GetCurrentAsset())
 		{
-			UAnimSequence* AnimSequence = Cast<UAnimSequence>( PreviewInstance->CurrentAsset );
+			UAnimSequence* AnimSequence = Cast<UAnimSequence>( PreviewInstance->GetCurrentAsset() );
 			if( AnimSequence )
 			{
 				const FScopedTransaction Transaction( LOCTEXT("CropAnimSequence", "Crop Animation Sequence") );
@@ -462,9 +462,9 @@ void SAnimationScrubPanel::OnCropAnimSequence( bool bFromStart, float CurrentTim
 void SAnimationScrubPanel::OnAppendAnimSequence( bool bFromStart, int32 NumOfFrames )
 {
 	UAnimSingleNodeInstance* PreviewInstance = GetPreviewInstance();
-	if(PreviewInstance && PreviewInstance->CurrentAsset)
+	if(PreviewInstance && PreviewInstance->GetCurrentAsset())
 	{
-		UAnimSequence* AnimSequence = Cast<UAnimSequence>(PreviewInstance->CurrentAsset);
+		UAnimSequence* AnimSequence = Cast<UAnimSequence>(PreviewInstance->GetCurrentAsset());
 		if(AnimSequence)
 		{
 			const FScopedTransaction Transaction(LOCTEXT("InsertAnimSequence", "Insert Animation Sequence"));
@@ -491,9 +491,9 @@ void SAnimationScrubPanel::OnAppendAnimSequence( bool bFromStart, int32 NumOfFra
 void SAnimationScrubPanel::OnInsertAnimSequence( bool bBefore, int32 CurrentFrame )
 {
 	UAnimSingleNodeInstance* PreviewInstance = GetPreviewInstance();
-	if(PreviewInstance && PreviewInstance->CurrentAsset)
+	if(PreviewInstance && PreviewInstance->GetCurrentAsset())
 	{
-		UAnimSequence* AnimSequence = Cast<UAnimSequence>(PreviewInstance->CurrentAsset);
+		UAnimSequence* AnimSequence = Cast<UAnimSequence>(PreviewInstance->GetCurrentAsset());
 		if(AnimSequence)
 		{
 			const FScopedTransaction Transaction(LOCTEXT("InsertAnimSequence", "Insert Animation Sequence"));
@@ -523,9 +523,9 @@ void SAnimationScrubPanel::OnReZeroAnimSequence( )
 	{
 		UDebugSkelMeshComponent* PreviewSkelComp = PersonaPtr.Pin()->GetPreviewMeshComponent();
 
-		if (PreviewInstance->CurrentAsset && PreviewSkelComp )
+		if (PreviewInstance->GetCurrentAsset() && PreviewSkelComp )
 		{
-			UAnimSequence* AnimSequence = Cast<UAnimSequence>( PreviewInstance->CurrentAsset );
+			UAnimSequence* AnimSequence = Cast<UAnimSequence>( PreviewInstance->GetCurrentAsset() );
 			if( AnimSequence )
 			{
 				const FScopedTransaction Transaction( LOCTEXT("ReZeroAnimation", "ReZero Animation Sequence") );
diff --git a/Engine/Source/Editor/Persona/Private/SMorphTargetViewer.cpp b/Engine/Source/Editor/Persona/Private/SMorphTargetViewer.cpp
--- a/Engine/Source/Editor/Persona/Private/SMorphTargetViewer.cpp
+++ b/Engine/Source/Editor/Persona/Private/SMorphTargetViewer.cpp
@@ -273,7 +273,7 @@ float SMorphTargetListRow::GetWeight() const
 		if ( SkelComp && SkelComp->GetAnimInstance() )
 		{
 			// make sure if they have value that's not same as saved value
-			float* CurrentValue = SkelComp->GetAnimInstance()->MorphTargetCurves.Find(Item->Name);
+			float* CurrentValue = SkelComp->GetAnimInstance()->GetMorphTargetCurves().Find(Item->Name);
 			if (CurrentValue && *CurrentValue != 0.f)
 			{
 				return *CurrentValue;
diff --git a/Engine/Source/Editor/Persona/Private/SRetargetManager.cpp b/Engine/Source/Editor/Persona/Private/SRetargetManager.cpp
--- a/Engine/Source/Editor/Persona/Private/SRetargetManager.cpp
+++ b/Engine/Source/Editor/Persona/Private/SRetargetManager.cpp
@@ -217,12 +217,12 @@ FReply SRetargetManager::OnViewRetargetBasePose()
 	if (PreviewMeshComp && PreviewMeshComp->PreviewInstance)
 	{
 		const FScopedTransaction Transaction(LOCTEXT("ViewRetargetBasePose_Action", "Edit Retarget Base Pose"));
-		PreviewMeshComp->PreviewInstance->bForceRetargetBasePose = !PreviewMeshComp->PreviewInstance->bForceRetargetBasePose;
+		PreviewMeshComp->PreviewInstance->SetForceRetargetBasePose(!PreviewMeshComp->PreviewInstance->GetForceRetargetBasePose());
 		PreviewMeshComp->Modify();
 		// reset all bone transform since you don't want to keep any bone transform change
 		PreviewMeshComp->PreviewInstance->ResetModifiedBone();
 		// add root 
-		if (PreviewMeshComp->PreviewInstance->bForceRetargetBasePose)
+		if (PreviewMeshComp->PreviewInstance->GetForceRetargetBasePose())
 		{
 			PreviewMeshComp->BonesOfInterest.Add(0);
 		}
@@ -276,7 +276,7 @@ FReply SRetargetManager::OnSaveRetargetBasePose()
 			// Clear PreviewMeshComp bone modified, they're baked now
 			PreviewMeshComp->PreviewInstance->ResetModifiedBone();
 			// turn off the retarget base pose if you're looking at it
-			PreviewMeshComp->PreviewInstance->bForceRetargetBasePose = false;
+			PreviewMeshComp->PreviewInstance->SetForceRetargetBasePose(false);
 		}
 	}
 	return FReply::Handled();
@@ -302,7 +302,7 @@ FReply SRetargetManager::OnResetRetargetBasePose()
 				// reset to original ref pose
 				PreviewMesh->RetargetBasePose = PreviewMesh->RefSkeleton.GetRefBonePose();
 				// turn off the retarget base pose if you're looking at it
-				PreviewMeshComp->PreviewInstance->bForceRetargetBasePose = true;
+				PreviewMeshComp->PreviewInstance->SetForceRetargetBasePose(true);
 			}
 		}
 	}
@@ -319,7 +319,7 @@ FText SRetargetManager::GetToggleRetargetBasePose() const
 	UDebugSkelMeshComponent * PreviewMeshComp = PersonaPtr.Pin()->GetPreviewMeshComponent();
 	if(PreviewMeshComp && PreviewMeshComp->PreviewInstance)
 	{
-		if (PreviewMeshComp->PreviewInstance->bForceRetargetBasePose)
+		if (PreviewMeshComp->PreviewInstance->GetForceRetargetBasePose())
 		{
 			return LOCTEXT("HideRetargetBasePose_Label", "Hide Pose");
 		}
diff --git a/Engine/Source/Editor/PropertyEditor/Private/PropertyHandleImpl.cpp b/Engine/Source/Editor/PropertyEditor/Private/PropertyHandleImpl.cpp
--- a/Engine/Source/Editor/PropertyEditor/Private/PropertyHandleImpl.cpp
+++ b/Engine/Source/Editor/PropertyEditor/Private/PropertyHandleImpl.cpp
@@ -334,7 +334,7 @@ FPropertyAccess::Result FPropertyValueImpl::ImportText( const TArray<FObjectBase
 		// GWorld.  Assumes all objects managed by this property window belong to the same world.
 		if (UPackage* ObjectPackage = (InObjects[0].Object ? InObjects[0].Object->GetOutermost() : nullptr))
 		{
-			bool bIsPIEPackage = !!(ObjectPackage->PackageFlags & PKG_PlayInEditor);
+			const bool bIsPIEPackage = ObjectPackage->HasAnyPackageFlags(PKG_PlayInEditor);
 			if (GUnrealEd && GUnrealEd->PlayWorld && bIsPIEPackage && !GIsPlayInEditorWorld)
 			{
 				OldGWorld = SetPlayInEditorWorld(GUnrealEd->PlayWorld);
diff --git a/Engine/Source/Editor/SceneOutliner/Private/SSceneOutliner.cpp b/Engine/Source/Editor/SceneOutliner/Private/SSceneOutliner.cpp
--- a/Engine/Source/Editor/SceneOutliner/Private/SSceneOutliner.cpp
+++ b/Engine/Source/Editor/SceneOutliner/Private/SSceneOutliner.cpp
@@ -1584,7 +1584,7 @@ namespace SceneOutliner
 				bool bAnyInPIE = false;
 				for (auto* Actor : SelectedActors)
 				{
-					if (!bAnyInPIE && Actor && (Actor->GetOutermost()->PackageFlags & PKG_PlayInEditor) != 0)
+					if (!bAnyInPIE && Actor && Actor->GetOutermost()->HasAnyPackageFlags(PKG_PlayInEditor))
 					{
 						bAnyInPIE = true;
 					}
@@ -1598,7 +1598,7 @@ namespace SceneOutliner
 				for (FSelectionIterator SelectionIt( *GEditor->GetSelectedActors() ); SelectionIt && !bChanged; ++SelectionIt)
 				{
 					AActor* Actor = CastChecked< AActor >( *SelectionIt );
-					if (!bAnyInPIE && (Actor->GetOutermost()->PackageFlags & PKG_PlayInEditor) != 0)
+					if (!bAnyInPIE && Actor->GetOutermost()->HasAnyPackageFlags(PKG_PlayInEditor))
 					{
 						bAnyInPIE = true;
 					}
diff --git a/Engine/Source/Editor/Sequencer/Private/Sequencer.cpp b/Engine/Source/Editor/Sequencer/Private/Sequencer.cpp
--- a/Engine/Source/Editor/Sequencer/Private/Sequencer.cpp
+++ b/Engine/Source/Editor/Sequencer/Private/Sequencer.cpp
@@ -1599,7 +1599,7 @@ void FSequencer::PlaceActorInFrontOfCamera( AActor* ActorCDO )
 	if ((GCurrentLevelEditingViewportClient != nullptr) && GCurrentLevelEditingViewportClient->IsPerspective())
 	{
 		// Don't allow this when the active viewport is showing a simulation/PIE level
-		const bool bIsViewportShowingPIEWorld = ( GCurrentLevelEditingViewportClient->GetWorld()->GetOutermost()->PackageFlags & PKG_PlayInEditor ) != 0;
+		const bool bIsViewportShowingPIEWorld = GCurrentLevelEditingViewportClient->GetWorld()->GetOutermost()->HasAnyPackageFlags(PKG_PlayInEditor);
 		if( !bIsViewportShowingPIEWorld )
 		{
 			// @todo sequencer actors: Ideally we could use the actor's collision to figure out how far to push out
diff --git a/Engine/Source/Editor/UnrealEd/Private/Animation/DebugSkelMeshComponent.cpp b/Engine/Source/Editor/UnrealEd/Private/Animation/DebugSkelMeshComponent.cpp
--- a/Engine/Source/Editor/UnrealEd/Private/Animation/DebugSkelMeshComponent.cpp
+++ b/Engine/Source/Editor/UnrealEd/Private/Animation/DebugSkelMeshComponent.cpp
@@ -246,23 +246,25 @@ FString UDebugSkelMeshComponent::GetPreviewText() const
 
 	if (IsPreviewOn())
 	{
-		if (UBlendSpaceBase* BlendSpace = Cast<UBlendSpaceBase>(PreviewInstance->CurrentAsset))
+		UAnimationAsset* CurrentAsset = PreviewInstance->GetCurrentAsset();
+		UVertexAnimation* CurrentVertexAnim = PreviewInstance->GetCurrentVertexAnimation();
+		if (UBlendSpaceBase* BlendSpace = Cast<UBlendSpaceBase>(CurrentAsset))
 		{
 			return FText::Format( LOCTEXT("BlendSpace", "Blend Space {0}"), FText::FromString(BlendSpace->GetName()) ).ToString();
 		}
-		else if (UAnimMontage* Montage = Cast<UAnimMontage>(PreviewInstance->CurrentAsset))
+		else if (UAnimMontage* Montage = Cast<UAnimMontage>(CurrentAsset))
 		{
 			return FText::Format( LOCTEXT("Montage", "Montage {0}"), FText::FromString(Montage->GetName()) ).ToString();
 		}
-		else if(UAnimComposite* Composite = Cast<UAnimComposite>(PreviewInstance->CurrentAsset))
+		else if(UAnimComposite* Composite = Cast<UAnimComposite>(CurrentAsset))
 		{
 			return FText::Format(LOCTEXT("Composite", "Composite {0}"), FText::FromString(Composite->GetName())).ToString();
 		}
-		else if (UAnimSequence* Sequence = Cast<UAnimSequence>(PreviewInstance->CurrentAsset))
+		else if (UAnimSequence* Sequence = Cast<UAnimSequence>(CurrentAsset))
 		{
 			return FText::Format( LOCTEXT("Animation", "Animation {0}"), FText::FromString(Sequence->GetName()) ).ToString();
 		}
-		else if (UVertexAnimation* VertAnim = Cast<UVertexAnimation>(PreviewInstance->CurrentVertexAnim))
+		else if (UVertexAnimation* VertAnim = Cast<UVertexAnimation>(CurrentVertexAnim))
 		{
 			return FText::Format( LOCTEXT("VertexAnim", "VertexAnim {0}"), FText::FromString(VertAnim->GetName()) ).ToString();
 		}
@@ -277,9 +279,9 @@ void UDebugSkelMeshComponent::InitAnim(bool bForceReinit)
 {
 	// If we already have PreviewInstnace and its asset's Skeleton does not match with mesh's Skeleton
 	// then we need to clear it up to avoid an issue
-	if ( PreviewInstance && PreviewInstance->CurrentAsset && SkeletalMesh )
+	if ( PreviewInstance && PreviewInstance->GetCurrentAsset() && SkeletalMesh )
 	{
-		if ( PreviewInstance->CurrentAsset->GetSkeleton() != SkeletalMesh->Skeleton )
+		if ( PreviewInstance->GetCurrentAsset()->GetSkeleton() != SkeletalMesh->Skeleton )
 		{
 			// if it doesn't match, just clear it
 			PreviewInstance->SetAnimationAsset(NULL);
@@ -347,7 +349,7 @@ void UDebugSkelMeshComponent::EnablePreview(bool bEnable, UAnimationAsset* Previ
 			}
 		else if (IsPreviewOn())
 		{
-			if (PreviewInstance->CurrentAsset == PreviewAsset || PreviewAsset == NULL)
+			if (PreviewInstance->GetCurrentAsset() == PreviewAsset || PreviewAsset == NULL)
 			{
 				// now recover to saved AnimScriptInstance;
 				AnimScriptInstance = SavedAnimScriptInstance;
@@ -431,62 +433,62 @@ void UDebugSkelMeshComponent::RefreshBoneTransforms(FActorComponentTickFunction*
 	const bool bIsPreviewInstance = (PreviewInstance && PreviewInstance == AnimScriptInstance);
 
 	BakedAnimationPoses.Reset();
-	if(bDisplayBakedAnimation && bIsPreviewInstance && PreviewInstance->RequiredBones.IsValid())
+	if(bDisplayBakedAnimation && bIsPreviewInstance && PreviewInstance->GetRequiredBones().IsValid())
 	{
-		if(UAnimSequence* Sequence = Cast<UAnimSequence>(PreviewInstance->CurrentAsset))
+		if(UAnimSequence* Sequence = Cast<UAnimSequence>(PreviewInstance->GetCurrentAsset()))
 		{
-			BakedAnimationPoses.AddUninitialized(PreviewInstance->RequiredBones.GetNumBones());
-			bool bSavedUseSourceData = AnimScriptInstance->RequiredBones.ShouldUseSourceData();
-			AnimScriptInstance->RequiredBones.SetUseRAWData(true);
-			AnimScriptInstance->RequiredBones.SetUseSourceData(false);
+			BakedAnimationPoses.AddUninitialized(PreviewInstance->GetRequiredBones().GetNumBones());
+			bool bSavedUseSourceData = PreviewInstance->GetRequiredBones().ShouldUseSourceData();
+			PreviewInstance->GetRequiredBones().SetUseRAWData(true);
+			PreviewInstance->GetRequiredBones().SetUseSourceData(false);
 			PreviewInstance->EnableControllers(false);
 			GenSpaceBases(BakedAnimationPoses);
-			AnimScriptInstance->RequiredBones.SetUseRAWData(false);
-			AnimScriptInstance->RequiredBones.SetUseSourceData(bSavedUseSourceData);
+			PreviewInstance->GetRequiredBones().SetUseRAWData(false);
+			PreviewInstance->GetRequiredBones().SetUseSourceData(bSavedUseSourceData);
 			PreviewInstance->EnableControllers(true);
 		}
 	}
 
 	SourceAnimationPoses.Reset();
-	if(bDisplaySourceAnimation && bIsPreviewInstance && PreviewInstance->RequiredBones.IsValid())
+	if(bDisplaySourceAnimation && bIsPreviewInstance && PreviewInstance->GetRequiredBones().IsValid())
 	{
-		if(UAnimSequence* Sequence = Cast<UAnimSequence>(PreviewInstance->CurrentAsset))
+		if(UAnimSequence* Sequence = Cast<UAnimSequence>(PreviewInstance->GetCurrentAsset()))
 		{
-			SourceAnimationPoses.AddUninitialized(PreviewInstance->RequiredBones.GetNumBones());
-			bool bSavedUseSourceData = AnimScriptInstance->RequiredBones.ShouldUseSourceData();
-			AnimScriptInstance->RequiredBones.SetUseSourceData(true);
+			SourceAnimationPoses.AddUninitialized(PreviewInstance->GetRequiredBones().GetNumBones());
+			bool bSavedUseSourceData = PreviewInstance->GetRequiredBones().ShouldUseSourceData();
+			PreviewInstance->GetRequiredBones().SetUseSourceData(true);
 			PreviewInstance->EnableControllers(false);
 			GenSpaceBases(SourceAnimationPoses);
-			AnimScriptInstance->RequiredBones.SetUseSourceData(bSavedUseSourceData);
+			PreviewInstance->GetRequiredBones().SetUseSourceData(bSavedUseSourceData);
 			PreviewInstance->EnableControllers(true);
 		}
 	}
 
 	UncompressedSpaceBases.Reset();
-	if (bDisplayRawAnimation && AnimScriptInstance && AnimScriptInstance->RequiredBones.IsValid())
+	if (bDisplayRawAnimation && AnimScriptInstance && AnimScriptInstance->GetRequiredBones().IsValid())
 	{
-		UncompressedSpaceBases.AddUninitialized(AnimScriptInstance->RequiredBones.GetNumBones());
+		UncompressedSpaceBases.AddUninitialized(AnimScriptInstance->GetRequiredBones().GetNumBones());
 
-		AnimScriptInstance->RequiredBones.SetUseRAWData(true);
+		AnimScriptInstance->GetRequiredBones().SetUseRAWData(true);
 		GenSpaceBases(UncompressedSpaceBases);
-		AnimScriptInstance->RequiredBones.SetUseRAWData(false);
+		AnimScriptInstance->GetRequiredBones().SetUseRAWData(false);
 	}
 
 	// Non retargeted pose.
 	NonRetargetedSpaceBases.Reset();
-	if( bDisplayNonRetargetedPose && AnimScriptInstance && AnimScriptInstance->RequiredBones.IsValid() )
+	if( bDisplayNonRetargetedPose && AnimScriptInstance && AnimScriptInstance->GetRequiredBones().IsValid() )
 	{
-		NonRetargetedSpaceBases.AddUninitialized(AnimScriptInstance->RequiredBones.GetNumBones());
-		AnimScriptInstance->RequiredBones.SetDisableRetargeting(true);
+		NonRetargetedSpaceBases.AddUninitialized(AnimScriptInstance->GetRequiredBones().GetNumBones());
+		AnimScriptInstance->GetRequiredBones().SetDisableRetargeting(true);
 		GenSpaceBases(NonRetargetedSpaceBases);
-		AnimScriptInstance->RequiredBones.SetDisableRetargeting(false);
+		AnimScriptInstance->GetRequiredBones().SetDisableRetargeting(false);
 	}
 
 	// Only works in PreviewInstance, and not for anim blueprint. This is intended.
 	AdditiveBasePoses.Reset();
-	if( bDisplayAdditiveBasePose && bIsPreviewInstance && PreviewInstance->RequiredBones.IsValid() )
+	if( bDisplayAdditiveBasePose && bIsPreviewInstance && PreviewInstance->GetRequiredBones().IsValid() )
 	{
-		if (UAnimSequence* Sequence = Cast<UAnimSequence>(PreviewInstance->CurrentAsset)) 
+		if (UAnimSequence* Sequence = Cast<UAnimSequence>(PreviewInstance->GetCurrentAsset())) 
 		{ 
 			if (Sequence->IsValidAdditive()) 
 			{ 
@@ -494,15 +496,15 @@ void UDebugSkelMeshComponent::RefreshBoneTransforms(FActorComponentTickFunction*
 				{
 					FCompactPose AdditiveBasePose;
 					FBlendedCurve AdditiveCurve(PreviewInstance);
-					AdditiveBasePose.SetBoneContainer(&PreviewInstance->RequiredBones);
-					Sequence->GetAdditiveBasePose(AdditiveBasePose, AdditiveCurve, FAnimExtractContext(PreviewInstance->CurrentTime));
+					AdditiveBasePose.SetBoneContainer(&PreviewInstance->GetRequiredBones());
+					Sequence->GetAdditiveBasePose(AdditiveBasePose, AdditiveCurve, FAnimExtractContext(PreviewInstance->GetCurrentTime()));
 					CSAdditiveBasePose.InitPose(AdditiveBasePose);
 				}
 
-				AdditiveBasePoses.AddUninitialized(PreviewInstance->RequiredBones.GetNumBones());
+				AdditiveBasePoses.AddUninitialized(PreviewInstance->GetRequiredBones().GetNumBones());
 				for (int32 i = 0; i < AdditiveBasePoses.Num(); ++i)
 				{
-					FCompactPoseBoneIndex CompactIndex = PreviewInstance->RequiredBones.MakeCompactPoseIndex(FMeshPoseBoneIndex(i));
+					FCompactPoseBoneIndex CompactIndex = PreviewInstance->GetRequiredBones().MakeCompactPoseIndex(FMeshPoseBoneIndex(i));
 					AdditiveBasePoses[i] = CSAdditiveBasePose.GetComponentSpaceTransform(CompactIndex);
 				}
 			}
diff --git a/Engine/Source/Editor/UnrealEd/Private/Commandlets/ContentCommandlets.cpp b/Engine/Source/Editor/UnrealEd/Private/Commandlets/ContentCommandlets.cpp
--- a/Engine/Source/Editor/UnrealEd/Private/Commandlets/ContentCommandlets.cpp
+++ b/Engine/Source/Editor/UnrealEd/Private/Commandlets/ContentCommandlets.cpp
@@ -391,7 +391,7 @@ void UResavePackagesCommandlet::LoadAndSaveOnePackage(const FString& Filename)
 			if (bStripEditorOnlyContent)
 			{
 				UE_LOG(LogContentCommandlet, Log, TEXT("Removing editor only data"));
-				Package->PackageFlags |= PKG_FilterEditorOnly;
+				Package->SetPackageFlags(PKG_FilterEditorOnly);
 			}
 
 			if (bSavePackage == true)
@@ -992,7 +992,7 @@ void ReferenceObject(UObject* Object, TMap<FString, FPackageObjects>& ObjectRefs
 	bool bIsValid = true;
 	// can't be in a script packge or be a field/template in a native package, or a top level pacakge, or in the transient package
 	if (!bValidWritableOuterName ||
-		(Object->GetOutermost()->PackageFlags & PKG_ContainsScript) ||
+		Object->GetOutermost()->HasAnyPackageFlags(PKG_ContainsScript) ||
 		Object->IsA(UField::StaticClass()) ||
 		Object->IsTemplate(RF_ClassDefaultObject) ||
 		Object->GetOuter() == NULL ||
diff --git a/Engine/Source/Editor/UnrealEd/Private/Commandlets/CookCommandlet.cpp b/Engine/Source/Editor/UnrealEd/Private/Commandlets/CookCommandlet.cpp
--- a/Engine/Source/Editor/UnrealEd/Private/Commandlets/CookCommandlet.cpp
+++ b/Engine/Source/Editor/UnrealEd/Private/Commandlets/CookCommandlet.cpp
@@ -349,14 +349,14 @@ bool UCookCommandlet::SaveCookedPackage( UPackage* Package, uint32 SaveFlags, bo
 		// Use SandboxFile to do path conversion to properly handle sandbox paths (outside of standard paths in particular).
 		Filename = SandboxFile->ConvertToAbsolutePathForExternalAppForWrite(*Filename);
 
-		uint32 OriginalPackageFlags = Package->PackageFlags;
+		uint32 OriginalPackageFlags = Package->GetPackageFlags();
 		UWorld* World = NULL;
 		EObjectFlags Flags = RF_NoFlags;
 		bool bPackageFullyLoaded = false;
 
 		if (bCompressed)
 		{
-			Package->PackageFlags |= PKG_StoreCompressed;
+			Package->SetPackageFlags(PKG_StoreCompressed);
 		}
 
 		ITargetPlatformManagerModule& TPM = GetTargetPlatformManagerRef();
@@ -440,11 +440,11 @@ bool UCookCommandlet::SaveCookedPackage( UPackage* Package, uint32 SaveFlags, bo
 
 				if (!Target->HasEditorOnlyData())
 				{
-					Package->PackageFlags |= PKG_FilterEditorOnly;
+					Package->SetPackageFlags(PKG_FilterEditorOnly);
 				}
 				else
 				{
-					Package->PackageFlags &= ~PKG_FilterEditorOnly;
+					Package->ClearPackageFlags(PKG_FilterEditorOnly);
 				}
 
 				if (World)
@@ -473,7 +473,7 @@ bool UCookCommandlet::SaveCookedPackage( UPackage* Package, uint32 SaveFlags, bo
 			}
 		}
 
-		Package->PackageFlags = OriginalPackageFlags;
+		Package->SetPackageFlagsTo(OriginalPackageFlags);
 	}
 
 	// return success
@@ -486,7 +486,7 @@ void UCookCommandlet::MaybeMarkPackageAsAlreadyLoaded(UPackage *Package)
 	if (PackagesToNotReload.Contains(Name))
 	{
 		UE_LOG(LogCookCommandlet, Verbose, TEXT("Marking %s already loaded."), *Name);
-		Package->PackageFlags |= PKG_ReloadingForCooker;
+		Package->SetPackageFlags(PKG_ReloadingForCooker);
 	}
 }
 
@@ -1400,7 +1400,7 @@ bool UCookCommandlet::Cook(const TArray<ITargetPlatform*>& Platforms, TArray<FSt
 					SaveCookedPackage(Pkg, SAVE_KeepGUID | SAVE_Async | (bUnversioned ? SAVE_Unversioned : 0), bWasUpToDate);
 
 					PackagesToNotReload.Add(Pkg->GetName());
-					Pkg->PackageFlags |= PKG_ReloadingForCooker;
+					Pkg->SetPackageFlags(PKG_ReloadingForCooker);
 					{
 						TArray<UObject *> ObjectsInPackage;
 						GetObjectsWithOuter(Pkg, ObjectsInPackage, true);
diff --git a/Engine/Source/Editor/UnrealEd/Private/Commandlets/DerivedDataCacheCommandlet.cpp b/Engine/Source/Editor/UnrealEd/Private/Commandlets/DerivedDataCacheCommandlet.cpp
--- a/Engine/Source/Editor/UnrealEd/Private/Commandlets/DerivedDataCacheCommandlet.cpp
+++ b/Engine/Source/Editor/UnrealEd/Private/Commandlets/DerivedDataCacheCommandlet.cpp
@@ -27,7 +27,7 @@ void UDerivedDataCacheCommandlet::MaybeMarkPackageAsAlreadyLoaded(UPackage *Pack
 	if (PackagesToNotReload.Contains(Name))
 	{
 		UE_LOG(LogDerivedDataCacheCommandlet, Verbose, TEXT("Marking %s already loaded."), *Name);
-		Package->PackageFlags |= PKG_ReloadingForCooker;
+		Package->SetPackageFlags(PKG_ReloadingForCooker);
 	}
 }
 
@@ -162,7 +162,7 @@ int32 UDerivedDataCacheCommandlet::Main( const FString& Params )
 							ProcessedPackages.Add(Filename);
 
 							PackagesToNotReload.Add(Pkg->GetName());
-							Pkg->PackageFlags |= PKG_ReloadingForCooker;
+							Pkg->SetPackageFlags(PKG_ReloadingForCooker);
 							{
 								TArray<UObject *> ObjectsInPackage;
 								GetObjectsWithOuter(Pkg, ObjectsInPackage, true);
diff --git a/Engine/Source/Editor/UnrealEd/Private/CookOnTheFlyServer.cpp b/Engine/Source/Editor/UnrealEd/Private/CookOnTheFlyServer.cpp
--- a/Engine/Source/Editor/UnrealEd/Private/CookOnTheFlyServer.cpp
+++ b/Engine/Source/Editor/UnrealEd/Private/CookOnTheFlyServer.cpp
@@ -2536,7 +2536,7 @@ ESavePackageResult UCookOnTheFlyServer::SaveCookedPackage(UPackage* Package, uin
 
 	if (Filename.Len())
 	{
-		if (Package->PackageFlags & PKG_ReloadingForCooker)
+		if (Package->HasAnyPackageFlags(PKG_ReloadingForCooker))
 		{
 			UE_LOG(LogCook, Warning, TEXT("Package %s marked as reloading for cook by was requested to save"), *Package->GetPathName());
 			UE_LOG(LogCook, Fatal, TEXT("Package %s marked as reloading for cook by was requested to save"), *Package->GetPathName());
@@ -2547,7 +2547,7 @@ ESavePackageResult UCookOnTheFlyServer::SaveCookedPackage(UPackage* Package, uin
 		// Use SandboxFile to do path conversion to properly handle sandbox paths (outside of standard paths in particular).
 		Filename = ConvertToFullSandboxPath(*Filename, true);
 
-		uint32 OriginalPackageFlags = Package->PackageFlags;
+		uint32 OriginalPackageFlags = Package->GetPackageFlags();
 		UWorld* World = NULL;
 		EObjectFlags Flags = RF_NoFlags;
 		bool bPackageFullyLoaded = false;
@@ -2600,7 +2600,7 @@ ESavePackageResult UCookOnTheFlyServer::SaveCookedPackage(UPackage* Package, uin
 
 		if (bShouldCompressPackage)
 		{
-			Package->PackageFlags |= PKG_StoreCompressed;
+			Package->SetPackageFlags(PKG_StoreCompressed);
 		}
 
 		for (ITargetPlatform* Target : Platforms)
@@ -2662,11 +2662,11 @@ ESavePackageResult UCookOnTheFlyServer::SaveCookedPackage(UPackage* Package, uin
 
 				if (!Target->HasEditorOnlyData())
 				{
-					Package->PackageFlags |= PKG_FilterEditorOnly;
+					Package->SetPackageFlags(PKG_FilterEditorOnly);
 				}
 				else
 				{
-					Package->PackageFlags &= ~PKG_FilterEditorOnly;
+					Package->ClearPackageFlags(PKG_FilterEditorOnly);
 				}
 
 				// need to subtract 32 because the SavePackage code creates temporary files with longer file names then the one we provide
@@ -2699,7 +2699,7 @@ ESavePackageResult UCookOnTheFlyServer::SaveCookedPackage(UPackage* Package, uin
 			}
 		}
 
-		Package->PackageFlags = OriginalPackageFlags;
+		Package->SetPackageFlagsTo(OriginalPackageFlags);
 	}
 
 	check( bIsSavingPackage == true );
@@ -4799,7 +4799,7 @@ void UCookOnTheFlyServer::MaybeMarkPackageAsAlreadyLoaded(UPackage *Package)
 	{
 		if (Package->IsFullyLoaded() == false)
 		{
-			Package->PackageFlags |= PKG_ReloadingForCooker;
+			Package->SetPackageFlags(PKG_ReloadingForCooker);
 		}
 	}
 
@@ -4809,7 +4809,7 @@ void UCookOnTheFlyServer::MaybeMarkPackageAsAlreadyLoaded(UPackage *Package)
 	if (PackagesToNotReload.Contains(Name))
 	{
 		UE_LOG(LogCookCommandlet, Verbose, TEXT("Marking %s already loaded."), *Name);
-		Package->PackageFlags |= PKG_ReloadingForCooker;
+		Package->SetPackageFlags(PKG_ReloadingForCooker);
 	}*/
 }
 
diff --git a/Engine/Source/Editor/UnrealEd/Private/DebugToolExec.cpp b/Engine/Source/Editor/UnrealEd/Private/DebugToolExec.cpp
--- a/Engine/Source/Editor/UnrealEd/Private/DebugToolExec.cpp
+++ b/Engine/Source/Editor/UnrealEd/Private/DebugToolExec.cpp
@@ -150,7 +150,7 @@ bool FDebugToolExec::Exec( UWorld* InWorld, const TCHAR* Cmd, FOutputDevice& Ar
 		if (foundObj != NULL)
 		{
 			// not allowed in the editor unless it is a PIE object as this command can have far reaching effects such as impacting serialization
-			if (!GIsEditor || (!foundObj->IsTemplate() && (foundObj->GetOutermost()->PackageFlags & PKG_PlayInEditor)))
+			if (!GIsEditor || (!foundObj->IsTemplate() && foundObj->GetOutermost()->HasAnyPackageFlags(PKG_PlayInEditor)))
 			{
 				EditObject(foundObj, true);
 			}
@@ -175,7 +175,7 @@ bool FDebugToolExec::Exec( UWorld* InWorld, const TCHAR* Cmd, FOutputDevice& Ar
 		if (foundObj != NULL)
 		{
 			// not allowed in the editor unless it is a PIE object as this command can have far reaching effects such as impacting serialization
-			if (!GIsEditor || (!foundObj->IsTemplate() && (foundObj->GetOutermost()->PackageFlags & PKG_PlayInEditor)))
+			if (!GIsEditor || (!foundObj->IsTemplate() && foundObj->GetOutermost()->HasAnyPackageFlags(PKG_PlayInEditor)))
 			{
 				EditObject(foundObj, false);
 			}
@@ -256,7 +256,7 @@ bool FDebugToolExec::Exec( UWorld* InWorld, const TCHAR* Cmd, FOutputDevice& Ar
 		if( Found )
 		{
 			// not allowed in the editor unless it is a PIE object as this command can have far reaching effects such as impacting serialization
-			if (!GIsEditor || (!Found->IsTemplate() && (Found->GetOutermost()->PackageFlags & PKG_PlayInEditor)))
+			if (!GIsEditor || (!Found->IsTemplate() && Found->GetOutermost()->HasAnyPackageFlags(PKG_PlayInEditor)))
 			{
 				EditObject(Found, true);
 			}
diff --git a/Engine/Source/Editor/UnrealEd/Private/Editor.cpp b/Engine/Source/Editor/UnrealEd/Private/Editor.cpp
--- a/Engine/Source/Editor/UnrealEd/Private/Editor.cpp
+++ b/Engine/Source/Editor/UnrealEd/Private/Editor.cpp
@@ -569,7 +569,7 @@ namespace EditorUtilities
 {
 	AActor* GetEditorWorldCounterpartActor( AActor* Actor )
 	{
-		const bool bIsSimActor = !!( Actor->GetOutermost()->PackageFlags & PKG_PlayInEditor );
+		const bool bIsSimActor = Actor->GetOutermost()->HasAnyPackageFlags(PKG_PlayInEditor);
 		if( bIsSimActor && GEditor->PlayWorld != NULL )
 		{
 			// Do we have a counterpart in the editor world?
@@ -600,7 +600,7 @@ namespace EditorUtilities
 
 	AActor* GetSimWorldCounterpartActor( AActor* Actor )
 	{
-		const bool bIsSimActor = !!( Actor->GetOutermost()->PackageFlags & PKG_PlayInEditor );
+		const bool bIsSimActor = Actor->GetOutermost()->HasAnyPackageFlags(PKG_PlayInEditor);
 		if( !bIsSimActor && GEditor->EditorWorld != NULL )
 		{
 			// Do we have a counterpart in the sim world?
@@ -826,7 +826,7 @@ namespace EditorUtilities
 			for (int32 ArrayIndex = 0; ArrayIndex < PropertyArrayDim; ArrayIndex++)
 			{
 				UObject* const SourceObjectPropertyValue = ObjectProperty->GetObjectPropertyValue_InContainer(InSourcePtr, ArrayIndex);
-				if (SourceObjectPropertyValue && SourceObjectPropertyValue->GetOutermost()->PackageFlags & PKG_PlayInEditor)
+				if (SourceObjectPropertyValue && SourceObjectPropertyValue->GetOutermost()->HasAnyPackageFlags(PKG_PlayInEditor))
 				{
 					// Not all the code paths below actually copy the object, but even if they don't we need to claim that they
 					// did, as copying a reference to an object in a PIE world leads to crashes
diff --git a/Engine/Source/Editor/UnrealEd/Private/EditorEngine.cpp b/Engine/Source/Editor/UnrealEd/Private/EditorEngine.cpp
--- a/Engine/Source/Editor/UnrealEd/Private/EditorEngine.cpp
+++ b/Engine/Source/Editor/UnrealEd/Private/EditorEngine.cpp
@@ -3749,17 +3749,13 @@ ESavePackageResult UEditorEngine::Save( UPackage* InOuter, UObject* InBase, EObj
 	}
 
 	UObject* Base = InBase;
-	if ( !Base && InOuter && InOuter->PackageFlags & PKG_ContainsMap )
+	if ( !Base && InOuter && InOuter->HasAnyPackageFlags(PKG_ContainsMap) )
 	{
 		Base = UWorld::FindWorldInPackage(InOuter);
 	}
 
 	// Record the package flags before OnPreSaveWorld. They will be used in OnPostSaveWorld.
-	uint32 OriginalPackageFlags = 0;
-	if ( InOuter )
-	{
-		OriginalPackageFlags = InOuter->PackageFlags;
-	}
+	const uint32 OriginalPackageFlags = (InOuter ? InOuter->GetPackageFlags() : 0);
 
 	SlowTask.EnterProgressFrame(10);
 
@@ -3911,7 +3907,7 @@ void UEditorEngine::OnPreSaveWorld(uint32 SaveFlags, UWorld* World)
 			// PIE prefix detected, mark package.
 			if( World->GetName().StartsWith( PLAYWORLD_PACKAGE_PREFIX ) )
 			{
-				World->GetOutermost()->PackageFlags |= PKG_PlayInEditor;
+				World->GetOutermost()->SetPackageFlags(PKG_PlayInEditor);
 			}
 		}
 		else
@@ -3980,10 +3976,10 @@ void UEditorEngine::OnPostSaveWorld(uint32 SaveFlags, UWorld* World, uint32 Orig
 		{
 			// Restore original value of PKG_PlayInEditor if we changed it during PIE saving
 			const bool bOriginallyPIE = (OriginalPackageFlags & PKG_PlayInEditor) != 0;
-			const bool bCurrentlyPIE = (WorldPackage->PackageFlags & PKG_PlayInEditor) != 0;
+			const bool bCurrentlyPIE = (WorldPackage->HasAnyPackageFlags(PKG_PlayInEditor));
 			if ( !bOriginallyPIE && bCurrentlyPIE )
 			{
-				WorldPackage->PackageFlags &= ~PKG_PlayInEditor;
+				WorldPackage->ClearPackageFlags(PKG_PlayInEditor);
 			}
 		}
 		else
diff --git a/Engine/Source/Editor/UnrealEd/Private/EditorServer.cpp b/Engine/Source/Editor/UnrealEd/Private/EditorServer.cpp
--- a/Engine/Source/Editor/UnrealEd/Private/EditorServer.cpp
+++ b/Engine/Source/Editor/UnrealEd/Private/EditorServer.cpp
@@ -2400,6 +2400,9 @@ bool UEditorEngine::Map_Load(const TCHAR* Str, FOutputDevice& Ar)
 				Context.SetCurrentWorld(World);
 				GWorld = World;
 
+				// UE-21181 - Tracking where the loaded editor level's package gets flagged as a PIE object
+				UPackage::EditorPackage = WorldPackage;
+
 				World->WorldType = EWorldType::Editor;
 
 				// Parse requested feature level if supplied
diff --git a/Engine/Source/Editor/UnrealEd/Private/EditorTransaction.cpp b/Engine/Source/Editor/UnrealEd/Private/EditorTransaction.cpp
--- a/Engine/Source/Editor/UnrealEd/Private/EditorTransaction.cpp
+++ b/Engine/Source/Editor/UnrealEd/Private/EditorTransaction.cpp
@@ -247,7 +247,7 @@ bool FTransaction::FObjectRecord::ContainsPieObject() const
 	{
 		UObject* Obj = Object.Get();
 
-		if(Obj && (Obj->GetOutermost()->PackageFlags & PKG_PlayInEditor) != 0)
+		if(Obj && Obj->GetOutermost()->HasAnyPackageFlags(PKG_PlayInEditor))
 		{
 			return true;
 		}
@@ -256,7 +256,7 @@ bool FTransaction::FObjectRecord::ContainsPieObject() const
 	for( const FReferencedObject& ReferencedObject : ReferencedObjects )
 	{
 		const UObject* Obj = ReferencedObject.GetObject();
-		if( Obj && (Obj->GetOutermost()->PackageFlags & PKG_PlayInEditor) != 0 )
+		if( Obj && Obj->GetOutermost()->HasAnyPackageFlags(PKG_PlayInEditor))
 		{
 			return true;
 		}
@@ -308,7 +308,7 @@ void FTransaction::SaveArray( UObject* Object, FScriptArray* Array, int32 Index,
 	check(Index+Count<=Array->Num());
 
 	// don't serialize the array if the object is contained within a PIE package
-	if( Object->HasAnyFlags(RF_Transactional) && (Object->GetOutermost()->PackageFlags&PKG_PlayInEditor) == 0 )
+	if( Object->HasAnyFlags(RF_Transactional) && !Object->GetOutermost()->HasAnyPackageFlags(PKG_PlayInEditor))
 	{
 		// Save the array.
 		new( Records )FObjectRecord( this, Object, Array, Index, Count, Oper, ElementSize, DefaultConstructor, Serializer, Destructor );
@@ -792,7 +792,7 @@ void UTransBuffer::SetPrimaryUndoObject(UObject* PrimaryObject)
 {
 	// Only record the primary object if its transactional, not in any of the temporary packages and theres an active transaction
 	if ( PrimaryObject && PrimaryObject->HasAnyFlags( RF_Transactional ) &&
-		( (PrimaryObject->GetOutermost()->PackageFlags & ( PKG_PlayInEditor|PKG_ContainsScript|PKG_CompiledIn ) ) == 0 ) )
+		(PrimaryObject->GetOutermost()->HasAnyPackageFlags( PKG_PlayInEditor|PKG_ContainsScript|PKG_CompiledIn ) == false) )
 	{
 		const int32 NumTransactions = UndoBuffer.Num();
 		const int32 CurrentTransactionIdx = NumTransactions - (UndoCount + 1);
diff --git a/Engine/Source/Editor/UnrealEd/Private/FileHelpers.cpp b/Engine/Source/Editor/UnrealEd/Private/FileHelpers.cpp
--- a/Engine/Source/Editor/UnrealEd/Private/FileHelpers.cpp
+++ b/Engine/Source/Editor/UnrealEd/Private/FileHelpers.cpp
@@ -2783,7 +2783,7 @@ bool FEditorFileUtils::SaveDirtyContentPackages(TArray<UClass*>& SaveContentClas
 		// Don't try to save "Transient" package.
 		bShouldIgnorePackage |= Package == GetTransientPackage();
 		// Ignore PIE packages.
-		bShouldIgnorePackage |= (Package->PackageFlags & PKG_PlayInEditor) != 0;
+		bShouldIgnorePackage |= Package->HasAnyPackageFlags(PKG_PlayInEditor);
 		// Ignore packages that haven't been modified.
 		bShouldIgnorePackage |= !Package->IsDirty();
 
@@ -2908,7 +2908,7 @@ FEditorFileUtils::EPromptReturnCode FEditorFileUtils::PromptForCheckoutAndSave(
 				if ( CurPackage != GetTransientPackage() )
 				{
 					// Never save compiled in packages
-					if ( (CurPackage->PackageFlags & PKG_CompiledIn) == 0 )
+					if (CurPackage->HasAnyPackageFlags(PKG_CompiledIn) == false)
 					{
 						if (UncheckedPackages.Contains(TWeakObjectPtr<UPackage>(CurPackage)))
 						{
@@ -2985,7 +2985,7 @@ FEditorFileUtils::EPromptReturnCode FEditorFileUtils::PromptForCheckoutAndSave(
 				if ( CurPackage != GetTransientPackage() )
 				{
 					// Never save compiled in packages
-					if ( (CurPackage->PackageFlags & PKG_CompiledIn) == 0 )
+					if (CurPackage->HasAnyPackageFlags(PKG_CompiledIn) == false)
 					{
 						FilteredPackages.Add( CurPackage );
 					}
@@ -3414,7 +3414,7 @@ void FEditorFileUtils::GetDirtyWorldPackages(TArray<UPackage*>& OutDirtyPackages
 	for (TObjectIterator<UWorld> WorldIt; WorldIt; ++WorldIt)
 	{
 		UPackage* WorldPackage = WorldIt->GetOutermost();
-		if (WorldPackage->IsDirty() && (WorldPackage->PackageFlags & PKG_PlayInEditor) == 0
+		if (WorldPackage->IsDirty() && (WorldPackage->HasAnyPackageFlags(PKG_PlayInEditor) == false)
 			&& !WorldPackage->HasAnyFlags(RF_Transient))
 		{
 			// IF the package is dirty and its not a pie package, add the world package to the list of packages to save
@@ -3438,7 +3438,7 @@ void FEditorFileUtils::GetDirtyContentPackages(TArray<UPackage*>& OutDirtyPackag
 		// Don't try to save packages with the RF_Transient flag.
 		bShouldIgnorePackage |= Package->HasAnyFlags(RF_Transient);
 		// Ignore PIE packages.
-		bShouldIgnorePackage |= (Package->PackageFlags & PKG_PlayInEditor) != 0;
+		bShouldIgnorePackage |= Package->HasAnyPackageFlags(PKG_PlayInEditor);
 		// Ignore packages that haven't been modified.
 		bShouldIgnorePackage |= !Package->IsDirty();
 
diff --git a/Engine/Source/Editor/UnrealEd/Private/ObjectTools.cpp b/Engine/Source/Editor/UnrealEd/Private/ObjectTools.cpp
--- a/Engine/Source/Editor/UnrealEd/Private/ObjectTools.cpp
+++ b/Engine/Source/Editor/UnrealEd/Private/ObjectTools.cpp
@@ -86,7 +86,7 @@ namespace ObjectTools
 			if( ObjectPackage != NULL )
 			{
 				if( ObjectPackage != GetTransientPackage()
-					&& (ObjectPackage->PackageFlags & PKG_PlayInEditor) == 0
+					&& (ObjectPackage->HasAnyPackageFlags(PKG_PlayInEditor) == false)
 					&& !Obj->IsPendingKill() )
 				{
 					bIsSupported = true;
@@ -1193,7 +1193,7 @@ namespace ObjectTools
 	void SelectActorsInLevelDirectlyReferencingObject( UObject* RefObj )
 	{
 		UPackage* Package = Cast<UPackage>(RefObj->GetOutermost());
-		if (Package && ((Package->PackageFlags & PKG_ContainsMap) != 0))
+		if (Package && Package->ContainsMap())
 		{
 			// Walk the chain of outers to find the object that is 'in' the level...
 			UObject* ObjToSelect = NULL;
@@ -2597,7 +2597,10 @@ namespace ObjectTools
 					NewPackage->GetOutermost()->FullyLoad();
 
 					// Make sure we copy all the cooked package flags if the asset was already cooked.
-					NewPackage->PackageFlags |= (Object->GetOutermost()->PackageFlags & PKG_FilterEditorOnly);
+					if (Object->GetOutermost()->HasAnyPackageFlags(PKG_FilterEditorOnly))
+					{
+						NewPackage->SetPackageFlags(PKG_FilterEditorOnly);
+					}
 					NewPackage->bIsCookedForEditor = Object->GetOutermost()->bIsCookedForEditor;
 
 					UObjectRedirector* Redirector = Cast<UObjectRedirector>( StaticFindObject(UObjectRedirector::StaticClass(), NewPackage, *NewObjectName) );
diff --git a/Engine/Source/Editor/UnrealEd/Private/PackageBackup.cpp b/Engine/Source/Editor/UnrealEd/Private/PackageBackup.cpp
--- a/Engine/Source/Editor/UnrealEd/Private/PackageBackup.cpp
+++ b/Engine/Source/Editor/UnrealEd/Private/PackageBackup.cpp
@@ -98,10 +98,10 @@ bool FAutoPackageBackup::ShouldBackupPackage( const UPackage& InPackage, FString
 	// Check various conditions to see if the package is a valid candidate for backing up
 	bool bShouldBackup =
 		GIsEditor																			// Backing up packages only makes sense in the editor
-		&& !IsRunningCommandlet()																			// Don't backup saves resulting from commandlets
+		&& !IsRunningCommandlet()															// Don't backup saves resulting from commandlets
 		&& IsPackageBackupEnabled()															// Ensure that the package backup is enabled in the first place
-		&& ( InPackage.PackageFlags & PKG_PlayInEditor ) == 0								// Don't back up PIE packages
-		&& ( InPackage.PackageFlags & PKG_ContainsScript ) == 0;							// Don't back up script packages
+		&& (InPackage.HasAnyPackageFlags(PKG_PlayInEditor) == false)						// Don't back up PIE packages
+		&& (InPackage.HasAnyPackageFlags(PKG_ContainsScript) == false);						// Don't back up script packages
 
 	if( bShouldBackup )
 	{
diff --git a/Engine/Source/Editor/UnrealEd/Private/PackageTools.cpp b/Engine/Source/Editor/UnrealEd/Private/PackageTools.cpp
--- a/Engine/Source/Editor/UnrealEd/Private/PackageTools.cpp
+++ b/Engine/Source/Editor/UnrealEd/Private/PackageTools.cpp
@@ -225,7 +225,7 @@ namespace PackageTools
 		// actor browser in case a script package was loaded
 		if ( Package != NULL )
 		{
-			if ( (Package->PackageFlags & PKG_ContainsScript) != 0 )
+			if (Package->HasAnyPackageFlags(PKG_ContainsScript))
 			{
 				GEditor->BroadcastClassPackageLoadedOrUnloaded();
 			}
@@ -330,7 +330,7 @@ namespace PackageTools
 
 				PackageBeingUnloaded->bHasBeenFullyLoaded = false;
 				PackageBeingUnloaded->ClearFlags(RF_WasLoaded);
-				if ( PackageBeingUnloaded->PackageFlags & PKG_ContainsScript )
+				if ( PackageBeingUnloaded->HasAnyPackageFlags(PKG_ContainsScript) )
 				{
 					bScriptPackageWasUnloaded = true;
 				}
diff --git a/Engine/Source/Editor/UnrealEd/Private/PlayLevel.cpp b/Engine/Source/Editor/UnrealEd/Private/PlayLevel.cpp
--- a/Engine/Source/Editor/UnrealEd/Private/PlayLevel.cpp
+++ b/Engine/Source/Editor/UnrealEd/Private/PlayLevel.cpp
@@ -223,7 +223,7 @@ void UEditorEngine::EndPlayMap()
 	{
 		UObject* Object = *It;
 
-		if ((Object->GetOutermost()->PackageFlags & PKG_PlayInEditor) != 0)
+		if (Object->GetOutermost()->HasAnyPackageFlags(PKG_PlayInEditor))
 		{
 			if (Object->HasAnyFlags(RF_Standalone))
 			{
@@ -364,7 +364,7 @@ void UEditorEngine::EndPlayMap()
 	for( FObjectIterator ObjectIt; ObjectIt; ++ObjectIt )
 	{
 		UObject* Object = *ObjectIt;
-		if( Object->GetOutermost()->PackageFlags & PKG_PlayInEditor )
+		if( Object->GetOutermost()->HasAnyPackageFlags(PKG_PlayInEditor))
 		{
 			UWorld* TheWorld = UWorld::FindWorldInPackage(Object->GetOutermost());
 			if ( TheWorld )
@@ -3444,7 +3444,7 @@ UWorld* UEditorEngine::CreatePIEWorldByDuplication(FWorldContext &WorldContext,
 	UE_LOG( LogPlayLevel, Log, TEXT("Creating play world package: %s"),  *PlayWorldMapName );	
 
 	UPackage* PlayWorldPackage = CastChecked<UPackage>(CreatePackage(NULL,*PlayWorldMapName));
-	PlayWorldPackage->PackageFlags |= PKG_PlayInEditor;
+	PlayWorldPackage->SetPackageFlags(PKG_PlayInEditor);
 	PlayWorldPackage->PIEInstanceID = WorldContext.PIEInstance;
 	PlayWorldPackage->FileName = InPackage->FileName;
 	PlayWorldPackage->SetGuid( InPackage->GetGuid() );
diff --git a/Engine/Source/Editor/UnrealEd/Private/SSkeletonWidget.cpp b/Engine/Source/Editor/UnrealEd/Private/SSkeletonWidget.cpp
--- a/Engine/Source/Editor/UnrealEd/Private/SSkeletonWidget.cpp
+++ b/Engine/Source/Editor/UnrealEd/Private/SSkeletonWidget.cpp
@@ -1383,7 +1383,7 @@ void SBasePoseViewport::SetSkeleton(USkeleton* Skeleton)
 				PreviewComponent->SetSkeletalMesh(PreviewSkeletalMesh);
 				PreviewComponent->EnablePreview(true, NULL, NULL);
 //				PreviewComponent->AnimScriptInstance = PreviewComponent->PreviewInstance;
-				PreviewComponent->PreviewInstance->bForceRetargetBasePose = true;
+				PreviewComponent->PreviewInstance->SetForceRetargetBasePose(true);
 				PreviewComponent->RefreshBoneTransforms(NULL);
 
 				//Place the camera at a good viewer position
diff --git a/Engine/Source/Editor/UnrealEd/Private/UnrealEdMisc.cpp b/Engine/Source/Editor/UnrealEd/Private/UnrealEdMisc.cpp
--- a/Engine/Source/Editor/UnrealEd/Private/UnrealEdMisc.cpp
+++ b/Engine/Source/Editor/UnrealEd/Private/UnrealEdMisc.cpp
@@ -1097,22 +1097,22 @@ void FUnrealEdMisc::OnMessageTokenActivated(const TSharedRef<IMessageToken>& Tok
 	}
 }
 
-FText FUnrealEdMisc::OnGetDisplayName(UObject* InObject, bool bFullPath)
+FText FUnrealEdMisc::OnGetDisplayName(const UObject* InObject, const bool bFullPath)
 {
 	FText Name = LOCTEXT("DisplayNone", "<None>");
 
-	if(InObject != NULL)
+	if (InObject != nullptr)
 	{
 		// Is this an object held by an actor?
-		AActor* Actor = NULL;
-		UActorComponent* Component = Cast<UActorComponent>(InObject);
+		const AActor* Actor = nullptr;
+		const UActorComponent* Component = Cast<UActorComponent>(InObject);
  
-		if (Component != NULL)
+		if (Component != nullptr)
 		{
 			Actor = Cast<AActor>(Component->GetOuter());
 		}
  
-		if (Actor != NULL)
+		if (Actor != nullptr)
 		{
 			Name = FText::FromString( bFullPath ? Actor->GetPathName() : Actor->GetName() );
 		}
diff --git a/Engine/Source/Editor/UnrealEd/Public/UnrealEdMisc.h b/Engine/Source/Editor/UnrealEd/Public/UnrealEdMisc.h
--- a/Engine/Source/Editor/UnrealEd/Public/UnrealEdMisc.h
+++ b/Engine/Source/Editor/UnrealEd/Public/UnrealEdMisc.h
@@ -258,7 +258,7 @@ private:
 	void OnMessageTokenActivated(const TSharedRef<class IMessageToken>& Token);
 
 	/** Delegate for (default) display name of UObject tokens. Can display the name of the actor if an object is/is part of one */
-	FText OnGetDisplayName(UObject* InObject, bool bFullPath);
+	FText OnGetDisplayName(const UObject* InObject, const bool bFullPath);
 
 	/** Delegate for (default) message log message selection - selects the objects that the tokens refer to (if any) */
 	void OnMessageSelectionChanged(TArray< TSharedRef<class FTokenizedMessage> >& Selection);
diff --git a/Engine/Source/Programs/UnrealHeaderTool/Private/CodeGenerator.cpp b/Engine/Source/Programs/UnrealHeaderTool/Private/CodeGenerator.cpp
--- a/Engine/Source/Programs/UnrealHeaderTool/Private/CodeGenerator.cpp
+++ b/Engine/Source/Programs/UnrealHeaderTool/Private/CodeGenerator.cpp
@@ -1208,8 +1208,8 @@ void FNativeClassHeaderGenerator::ExportGeneratedPackageInitCode(UPackage* InPac
 		GeneratedFunctionText.Logf(TEXT("#endif\r\n"));
 	}
 
-	GeneratedFunctionText.Logf(TEXT("\t\t\tReturnPackage->PackageFlags |= PKG_CompiledIn | 0x%08X;\r\n"), InPackage->PackageFlags & (PKG_ClientOptional | PKG_ServerSideOnly | PKG_EditorOnly));
-	TheFlagAudit.Add(InPackage, TEXT("PackageFlags"), InPackage->PackageFlags);
+	GeneratedFunctionText.Logf(TEXT("\t\t\tReturnPackage->SetPackageFlags(PKG_CompiledIn | 0x%08X);\r\n"), InPackage->GetPackageFlags() & (PKG_ClientOptional | PKG_ServerSideOnly | PKG_EditorOnly));
+	TheFlagAudit.Add(InPackage, TEXT("PackageFlags"), InPackage->GetPackageFlags());
 	{
 		FGuid Guid;
 
@@ -5412,12 +5412,11 @@ ECompilationResult::Type UnrealHeaderTool_Main(const FString& ModuleInfoFilename
 		// Set some package flags for indicating that this package contains script
 		// NOTE: We do this even if we didn't have to create the package, because CoreUObject is compiled into UnrealHeaderTool and we still
 		//       want to make sure our flags get set
-		Package->PackageFlags |= PKG_ContainsScript;
-		Package->PackageFlags &= ~(PKG_ClientOptional | PKG_ServerSideOnly);
-		Package->PackageFlags |= PKG_Compiling;
+		Package->SetPackageFlags(PKG_ContainsScript | PKG_Compiling);
+		Package->ClearPackageFlags(PKG_ClientOptional | PKG_ServerSideOnly);
 		if (Module.ModuleType == EBuildModuleType::Editor)
 		{
-			Package->PackageFlags |= PKG_EditorOnly;
+			Package->SetPackageFlags(PKG_EditorOnly);
 		}
 
 		GPackageToManifestModuleMap.Add(Package, &Module);
diff --git a/Engine/Source/Programs/UnrealHeaderTool/Private/HeaderParser.cpp b/Engine/Source/Programs/UnrealHeaderTool/Private/HeaderParser.cpp
--- a/Engine/Source/Programs/UnrealHeaderTool/Private/HeaderParser.cpp
+++ b/Engine/Source/Programs/UnrealHeaderTool/Private/HeaderParser.cpp
@@ -4484,6 +4484,12 @@ bool FHeaderParser::CompileDeclaration(FClasses& AllClasses, FUnrealSourceFile&
 		}
 	}
 
+	if (Token.Matches(TEXT("PRAGMA_DISABLE_DEPRECATION_WARNINGS")) || Token.Matches(TEXT("PRAGMA_ENABLE_DEPRECATION_WARNINGS")))
+	{
+		// Skip these macros
+		return true;
+	}
+
 	// Ignore C++ declaration / function definition. 
 	return SkipDeclaration(Token);
 }
diff --git a/Engine/Source/Runtime/AIModule/AIModule.Build.cs b/Engine/Source/Runtime/AIModule/AIModule.Build.cs
--- a/Engine/Source/Runtime/AIModule/AIModule.Build.cs
+++ b/Engine/Source/Runtime/AIModule/AIModule.Build.cs
@@ -33,6 +33,7 @@ namespace UnrealBuildTool.Rules
 				new string[] {
 					"RHI",
                     "RenderCore",
+                    "ShaderCore",
 				}
 				);
 
diff --git a/Engine/Source/Runtime/AIModule/Classes/BehaviorTree/BTFunctionLibrary.h b/Engine/Source/Runtime/AIModule/Classes/BehaviorTree/BTFunctionLibrary.h
--- a/Engine/Source/Runtime/AIModule/Classes/BehaviorTree/BTFunctionLibrary.h
+++ b/Engine/Source/Runtime/AIModule/Classes/BehaviorTree/BTFunctionLibrary.h
@@ -13,6 +13,11 @@ namespace FBTNodeBPImplementationHelper
 	static const int32 AISpecific = 1 << 1;
 	static const int32 All = Generic | AISpecific;
 
+	/** checks if given object implements GenericEventName and/or AIEventName BP events, and returns an result as flags set on return integer 
+	 *	@return flags set in returned integer indicate kinds of events implemented by given object */
+	AIMODULE_API int32 CheckEventImplementationVersion(FName GenericEventName, FName AIEventName, const UObject& Object, const UClass& StopAtClass);
+
+	DEPRECATED(4.11, "This version of CheckEventImplementationVersion is deprecated. Please use the one taking reference to UObject and StopAtClass rather than a pointers.")
 	AIMODULE_API int32 CheckEventImplementationVersion(FName GenericEventName, FName AIEventName, const UObject* Ob, const UClass* StopAtClass);
 }
 
diff --git a/Engine/Source/Runtime/AIModule/Classes/BehaviorTree/BehaviorTreeManager.h b/Engine/Source/Runtime/AIModule/Classes/BehaviorTree/BehaviorTreeManager.h
--- a/Engine/Source/Runtime/AIModule/Classes/BehaviorTree/BehaviorTreeManager.h
+++ b/Engine/Source/Runtime/AIModule/Classes/BehaviorTree/BehaviorTreeManager.h
@@ -41,7 +41,7 @@ class AIMODULE_API UBehaviorTreeManager : public UObject
 	static int32 GetAlignedDataSize(int32 Size);
 
 	/** helper function for sorting and aligning node memory */
-	static void InitializeMemoryHelper(const TArray<UBTDecorator*>& Nodes, TArray<uint16>& MemoryOffsets, int32& MemorySize);
+	static void InitializeMemoryHelper(const TArray<UBTDecorator*>& Nodes, TArray<uint16>& MemoryOffsets, int32& MemorySize, bool bForceInstancing = false);
 
 	/** cleanup hooks for map loading */
 	virtual void FinishDestroy() override;
diff --git a/Engine/Source/Runtime/AIModule/Classes/Navigation/CrowdFollowingComponent.h b/Engine/Source/Runtime/AIModule/Classes/Navigation/CrowdFollowingComponent.h
--- a/Engine/Source/Runtime/AIModule/Classes/Navigation/CrowdFollowingComponent.h
+++ b/Engine/Source/Runtime/AIModule/Classes/Navigation/CrowdFollowingComponent.h
@@ -84,6 +84,9 @@ class AIMODULE_API UCrowdFollowingComponent : public UPathFollowingComponent, pu
 	void SetCrowdAvoidanceRangeMultiplier(float Multipler, bool bUpdateAgent = true);
 	void SetCrowdAffectFallingVelocity(bool bEnable);
 	void SetCrowdRotateToVelocity(bool bEnable);
+	void SetAvoidanceGroup(int32 GroupFlags, bool bUpdateAgent = true);
+	void SetGroupsToAvoid(int32 GroupFlags, bool bUpdateAgent = true);
+	void SetGroupsToIgnore(int32 GroupFlags, bool bUpdateAgent = true);
 
 	FORCEINLINE bool IsCrowdSimulationEnabled() const { return bEnableCrowdSimulation; }
 	FORCEINLINE bool IsCrowdSimulatioSuspended() const { return bSuspendCrowdSimulation; }
@@ -200,6 +203,9 @@ protected:
 	bool HasMovedDuringPause() const;
 	void UpdateCachedDirections(const FVector& NewVelocity, const FVector& NextPathCorner, bool bTraversingLink);
 	virtual bool UpdateCachedGoal(FVector& NewGoalPos);
+	
+	void OnPendingNavigationInit();
+	bool RegisterCrowdAgent();
 
 	friend UCrowdManager;
 };
diff --git a/Engine/Source/Runtime/AIModule/Classes/Navigation/CrowdManager.h b/Engine/Source/Runtime/AIModule/Classes/Navigation/CrowdManager.h
--- a/Engine/Source/Runtime/AIModule/Classes/Navigation/CrowdManager.h
+++ b/Engine/Source/Runtime/AIModule/Classes/Navigation/CrowdManager.h
@@ -180,7 +180,7 @@ class AIMODULE_API UCrowdManager : public UObject
 	bool SetAgentMoveDirection(const UCrowdFollowingComponent* AgentComponent, const FVector& MoveDirection) const;
 
 	/** sets move target using path (only for fully simulated) */
-	bool SetAgentMovePath(const UCrowdFollowingComponent* AgentComponent, const FNavMeshPath* Path, int32 PathSectionStart, int32 PathSectionEnd) const;
+	bool SetAgentMovePath(const UCrowdFollowingComponent* AgentComponent, const FNavMeshPath* Path, int32 PathSectionStart, int32 PathSectionEnd, const FVector& PathSectionEndLocation) const;
 
 	/** clears move target for crowd agent (only for fully simulated) */
 	void ClearAgentMoveTarget(const UCrowdFollowingComponent* AgentComponent) const;
diff --git a/Engine/Source/Runtime/AIModule/Classes/Perception/AIPerceptionComponent.h b/Engine/Source/Runtime/AIModule/Classes/Perception/AIPerceptionComponent.h
--- a/Engine/Source/Runtime/AIModule/Classes/Perception/AIPerceptionComponent.h
+++ b/Engine/Source/Runtime/AIModule/Classes/Perception/AIPerceptionComponent.h
@@ -5,7 +5,10 @@
 #include "AISystem.h"
 #include "Perception/AISense.h"
 #include "AIPerceptionSystem.h"
-#include "Debug/DebugDrawService.h"
+#if !UE_BUILD_SHIPPING
+#	include "Debug/DebugDrawService.h"
+#	include "Debug/GameplayDebuggerBaseObject.h"
+#endif
 #include "AIPerceptionComponent.generated.h"
 
 class AAIController;
@@ -233,7 +236,7 @@ public:
 	bool HasActiveStimulus(const AActor& Source, FAISenseID Sense) const;
 
 #if !UE_BUILD_SHIPPING
-	void GetDebugData(TArray<FString>& OnScreenStrings, TArray<FDrawDebugShapeElement>& DebugShapes) const;
+	void GrabGameplayDebuggerData(TArray<FString>& OnScreenStrings, TArray<FGameplayDebuggerShapeElement>& DebugShapes) const;
 #endif // !UE_BUILD_SHIPPING
 
 #if ENABLE_VISUAL_LOG
diff --git a/Engine/Source/Runtime/AIModule/Classes/Perception/AIPerceptionSystem.h b/Engine/Source/Runtime/AIModule/Classes/Perception/AIPerceptionSystem.h
--- a/Engine/Source/Runtime/AIModule/Classes/Perception/AIPerceptionSystem.h
+++ b/Engine/Source/Runtime/AIModule/Classes/Perception/AIPerceptionSystem.h
@@ -134,6 +134,9 @@ public:
 	 */
 	void UnregisterSource(AActor& SourceActor, TSubclassOf<UAISense> Sense = nullptr);
 
+	void OnListenerForgetsActor(const UAIPerceptionComponent& Listener, AActor& ActorToForget);
+	void OnListenerForgetsAll(const UAIPerceptionComponent& Listener);
+
 	void RegisterDelayedStimulus(FPerceptionListenerID ListenerId, float Delay, AActor* Instigator, const FAIStimulus& Stimulus);
 
 	static UAIPerceptionSystem* GetCurrent(UObject* WorldContextObject);
diff --git a/Engine/Source/Runtime/AIModule/Classes/Perception/AISense.h b/Engine/Source/Runtime/AIModule/Classes/Perception/AISense.h
--- a/Engine/Source/Runtime/AIModule/Classes/Perception/AISense.h
+++ b/Engine/Source/Runtime/AIModule/Classes/Perception/AISense.h
@@ -40,6 +40,10 @@ protected:
 	/** If true all newly spawned pawns will get auto registered as source for this sense. */
 	UPROPERTY(EditDefaultsOnly, BlueprintReadOnly, Category = "AI Perception", config)
 	uint32 bAutoRegisterAllPawnsAsSources : 1;
+
+	/** this sense has some internal logic that requires it to be notified when 
+	 *	a listener wants to forget an actor*/
+	uint32 bNeedsForgettingNotification : 1;
 	
 private:
 	UPROPERTY()
@@ -105,6 +109,10 @@ public:
 	virtual void RegisterWrappedEvent(UAISenseEvent& PerceptionEvent);
 	virtual FAISenseID UpdateSenseID();
 
+	bool NeedsNotificationOnForgetting() const { return bNeedsForgettingNotification; }
+	virtual void OnListenerForgetsActor(const FPerceptionListener& Listener, AActor& ActorToForget) {}
+	virtual void OnListenerForgetsAll(const FPerceptionListener& Listener) {}
+
 	FORCEINLINE void OnNewListener(const FPerceptionListener& NewListener) { OnNewListenerDelegate.ExecuteIfBound(NewListener); }
 	FORCEINLINE void OnListenerUpdate(const FPerceptionListener& NewListener) { OnListenerUpdateDelegate.ExecuteIfBound(NewListener); }
 	FORCEINLINE void OnListenerRemoved(const FPerceptionListener& NewListener) { OnListenerRemovedDelegate.ExecuteIfBound(NewListener); }
diff --git a/Engine/Source/Runtime/AIModule/Classes/Perception/AISenseConfig.h b/Engine/Source/Runtime/AIModule/Classes/Perception/AISenseConfig.h
--- a/Engine/Source/Runtime/AIModule/Classes/Perception/AISenseConfig.h
+++ b/Engine/Source/Runtime/AIModule/Classes/Perception/AISenseConfig.h
@@ -4,7 +4,7 @@
 
 #include "AIPerceptionTypes.h"
 #if !UE_BUILD_SHIPPING
-#	include "Debug/DebugDrawService.h"
+#	include "Debug/GameplayDebuggerBaseObject.h"
 #endif
 #include "AISenseConfig.generated.h"
 
@@ -40,6 +40,6 @@ public:
 	//----------------------------------------------------------------------//
 	// DEBUG	
 	//----------------------------------------------------------------------//
-	virtual void GetDebugData(TArray<FString>& OnScreenStrings, TArray<FDrawDebugShapeElement>& DebugShapes, const UAIPerceptionComponent& PerceptionComponent) const {}
+	virtual void GetDebugData(TArray<FString>& OnScreenStrings, TArray<FGameplayDebuggerShapeElement>& DebugShapes, const UAIPerceptionComponent& PerceptionComponent) const {}
 #endif // !UE_BUILD_SHIPPING
 };
\ No newline at end of file
diff --git a/Engine/Source/Runtime/AIModule/Classes/Perception/AISenseConfig_Hearing.h b/Engine/Source/Runtime/AIModule/Classes/Perception/AISenseConfig_Hearing.h
--- a/Engine/Source/Runtime/AIModule/Classes/Perception/AISenseConfig_Hearing.h
+++ b/Engine/Source/Runtime/AIModule/Classes/Perception/AISenseConfig_Hearing.h
@@ -36,6 +36,6 @@ public:
 	//----------------------------------------------------------------------//
 	// DEBUG
 	//----------------------------------------------------------------------//
-	virtual void GetDebugData(TArray<FString>& OnScreenStrings, TArray<FDrawDebugShapeElement>& DebugShapes, const UAIPerceptionComponent& PerceptionComponent) const override;
+	virtual void GetDebugData(TArray<FString>& OnScreenStrings, TArray<FGameplayDebuggerShapeElement>& DebugShapes, const UAIPerceptionComponent& PerceptionComponent) const override;
 #endif // !UE_BUILD_SHIPPING
 };
diff --git a/Engine/Source/Runtime/AIModule/Classes/Perception/AISenseConfig_Sight.h b/Engine/Source/Runtime/AIModule/Classes/Perception/AISenseConfig_Sight.h
--- a/Engine/Source/Runtime/AIModule/Classes/Perception/AISenseConfig_Sight.h
+++ b/Engine/Source/Runtime/AIModule/Classes/Perception/AISenseConfig_Sight.h
@@ -41,6 +41,6 @@ public:
 	//----------------------------------------------------------------------//
 	// DEBUG
 	//----------------------------------------------------------------------//
-	virtual void GetDebugData(TArray<FString>& OnScreenStrings, TArray<FDrawDebugShapeElement>& DebugShapes, const UAIPerceptionComponent& PerceptionComponent) const override;
+	virtual void GetDebugData(TArray<FString>& OnScreenStrings, TArray<FGameplayDebuggerShapeElement>& DebugShapes, const UAIPerceptionComponent& PerceptionComponent) const override;
 #endif // !UE_BUILD_SHIPPING
 };
\ No newline at end of file
diff --git a/Engine/Source/Runtime/AIModule/Classes/Perception/AISense_Sight.h b/Engine/Source/Runtime/AIModule/Classes/Perception/AISense_Sight.h
--- a/Engine/Source/Runtime/AIModule/Classes/Perception/AISense_Sight.h
+++ b/Engine/Source/Runtime/AIModule/Classes/Perception/AISense_Sight.h
@@ -87,6 +87,12 @@ struct FAISightQuery
 		Score = Age + Importance;
 	}
 
+	void ForgetPreviousResult()
+	{
+		LastSeenLocation = FAISystem::InvalidLocation;
+		bLastResult = false;
+	}
+
 	class FSortPredicate
 	{
 	public:
@@ -156,6 +162,9 @@ public:
 	virtual void UnregisterSource(AActor& SourceActor) override;
 	virtual void CleanseInvalidSources() override;
 	
+	virtual void OnListenerForgetsActor(const FPerceptionListener& Listener, AActor& ActorToForget) override;
+	virtual void OnListenerForgetsAll(const FPerceptionListener& Listener) override;
+
 protected:
 	virtual float Update() override;
 
@@ -164,7 +173,7 @@ protected:
 	void OnNewListenerImpl(const FPerceptionListener& NewListener);
 	void OnListenerUpdateImpl(const FPerceptionListener& UpdatedListener);
 	void OnListenerRemovedImpl(const FPerceptionListener& UpdatedListener);	
-
+	
 	void GenerateQueriesForListener(const FPerceptionListener& Listener, const FDigestedSightProperties& PropertyDigest);
 
 	enum FQueriesOperationPostProcess
diff --git a/Engine/Source/Runtime/AIModule/Private/AIController.cpp b/Engine/Source/Runtime/AIModule/Private/AIController.cpp
--- a/Engine/Source/Runtime/AIModule/Private/AIController.cpp
+++ b/Engine/Source/Runtime/AIModule/Private/AIController.cpp
@@ -587,7 +587,7 @@ EPathFollowingRequestResult::Type AAIController::MoveTo(const FAIMoveRequest& Mo
 			const FNavAgentProperties& AgentProps = GetNavAgentPropertiesRef();
 			FNavLocation ProjectedLocation;
 
-			if (NavSys && !NavSys->ProjectPointToNavigation(MoveRequest.GetGoalLocation(), ProjectedLocation, AgentProps.GetExtent(), &AgentProps))
+			if (NavSys && !NavSys->ProjectPointToNavigation(MoveRequest.GetGoalLocation(), ProjectedLocation, INVALID_NAVEXTENT, &AgentProps))
 			{
 				UE_VLOG_LOCATION(this, LogAINavigation, Error, MoveRequest.GetGoalLocation(), 30.f, FColor::Red, TEXT("AAIController::MoveTo failed to project destination location to navmesh"));
 				bCanRequestMove = false;
diff --git a/Engine/Source/Runtime/AIModule/Private/AIModule.cpp b/Engine/Source/Runtime/AIModule/Private/AIModule.cpp
--- a/Engine/Source/Runtime/AIModule/Private/AIModule.cpp
+++ b/Engine/Source/Runtime/AIModule/Private/AIModule.cpp
@@ -40,7 +40,12 @@ void FAIModule::StartupModule()
 { 
 	// This code will execute after your module is loaded into memory and after global variables initialization. We needs some place to load GameplayDebugger module so it's best place for it now.
 #if !(UE_BUILD_SHIPPING || UE_BUILD_TEST)
-	FModuleManager::LoadModulePtr< IModuleInterface >("GameplayDebugger");
+	bool bEnableGameplayDebuggerPlugin = false;
+	GConfig->GetBool(TEXT("/Script/GameplayDebuggerModule.GameplayDebuggerModuleSettings"), TEXT("EnableGameplayDebuggerPlugin"), bEnableGameplayDebuggerPlugin, GEngineIni);
+	if (!bEnableGameplayDebuggerPlugin)
+	{
+		FModuleManager::LoadModulePtr< IModuleInterface >("GameplayDebugger");
+	}
 #endif
 
 #if WITH_EDITOR 
diff --git a/Engine/Source/Runtime/AIModule/Private/AIModulePrivate.h b/Engine/Source/Runtime/AIModule/Private/AIModulePrivate.h
--- a/Engine/Source/Runtime/AIModule/Private/AIModulePrivate.h
+++ b/Engine/Source/Runtime/AIModule/Private/AIModulePrivate.h
@@ -75,4 +75,5 @@
 #include "Tasks/AITask.h"
 
 DECLARE_LOG_CATEGORY_EXTERN(LogAINavigation, Warning, All);
-DECLARE_LOG_CATEGORY_EXTERN(LogCrowdFollowing, Warning, All);
\ No newline at end of file
+DECLARE_LOG_CATEGORY_EXTERN(LogCrowdFollowing, Warning, All);
+DECLARE_LOG_CATEGORY_EXTERN(LogAIModuleGameplayDebugger, Log, All);
diff --git a/Engine/Source/Runtime/AIModule/Private/Actions/PawnAction_BlueprintBase.cpp b/Engine/Source/Runtime/AIModule/Private/Actions/PawnAction_BlueprintBase.cpp
--- a/Engine/Source/Runtime/AIModule/Private/Actions/PawnAction_BlueprintBase.cpp
+++ b/Engine/Source/Runtime/AIModule/Private/Actions/PawnAction_BlueprintBase.cpp
@@ -3,19 +3,11 @@
 #include "AIModulePrivate.h"
 #include "Actions/PawnAction_BlueprintBase.h"
 
-namespace
-{
-	FORCEINLINE bool HasBlueprintFunction(FName FuncName, const UObject* Ob, const UClass* StopAtClass)
-	{
-		return (Ob->GetClass()->FindFunctionByName(FuncName)->GetOuter() != StopAtClass);
-	}
-}
-
 UPawnAction_BlueprintBase::UPawnAction_BlueprintBase(const FObjectInitializer& ObjectInitializer)
 	: Super(ObjectInitializer)
 {
 	UClass* StopAtClass = UPawnAction_BlueprintBase::StaticClass();
-	bWantsTick = HasBlueprintFunction(TEXT("ActionTick"), this, StopAtClass);
+	bWantsTick = BlueprintNodeHelpers::HasBlueprintFunction(TEXT("ActionTick"), *this, *StopAtClass);
 }
 
 void UPawnAction_BlueprintBase::Tick(float DeltaTime)
diff --git a/Engine/Source/Runtime/AIModule/Private/BehaviorTree/BTFunctionLibrary.cpp b/Engine/Source/Runtime/AIModule/Private/BehaviorTree/BTFunctionLibrary.cpp
--- a/Engine/Source/Runtime/AIModule/Private/BehaviorTree/BTFunctionLibrary.cpp
+++ b/Engine/Source/Runtime/AIModule/Private/BehaviorTree/BTFunctionLibrary.cpp
@@ -11,14 +11,20 @@
 //----------------------------------------------------------------------//
 namespace FBTNodeBPImplementationHelper
 {
-	int32 CheckEventImplementationVersion(FName GenericEventName, FName AIEventName, const UObject* Ob, const UClass* StopAtClass)
+	int32 CheckEventImplementationVersion(FName GenericEventName, FName AIEventName, const UObject& Object, const UClass& StopAtClass)
 	{
-		const bool bGeneric = BlueprintNodeHelpers::HasBlueprintFunction(GenericEventName, Ob, StopAtClass);
-		const bool bAI = BlueprintNodeHelpers::HasBlueprintFunction(AIEventName, Ob, StopAtClass);
-		
+		const bool bGeneric = BlueprintNodeHelpers::HasBlueprintFunction(GenericEventName, Object, StopAtClass);
+		const bool bAI = BlueprintNodeHelpers::HasBlueprintFunction(AIEventName, Object, StopAtClass);
+
 		return (bGeneric ? Generic : NoImplementation) | (bAI ? AISpecific : NoImplementation);
 	}
+
+	int32 CheckEventImplementationVersion(FName GenericEventName, FName AIEventName, const UObject* Ob, const UClass* StopAtClass)
+	{
+		return (Ob && StopAtClass) ? CheckEventImplementationVersion(GenericEventName, AIEventName, *Ob, *StopAtClass) : NoImplementation;
+	}
 }
+
 //----------------------------------------------------------------------//
 // 
 //----------------------------------------------------------------------//
