Integrating from Dev-Framework:

Fix CopyPropertiesForUnrelatedObjects clobbering anim BP compiler-generated property defaults & stabilize generated function names in the anim BP compiler.

The anim blueprint compiler generates property defaults (BoundFunction names) when compiling to allow anim graph nodes to call into their parent class functions. Also, each run of the anim blueprint compiler generated different names for this BoundFunction as the name generation bas based on the FName of visual anim nodes. This set up a load-order dependency where depending on how anim BPs were referenced by other data, we could potentially end up with mismatched function names in defaults and functions when compiling derived anim Blueprints.
The order of operations when the load order failed went something like this:

1. Compile base anim BP (in RegenerateBlueprintClass), generating new functions and adding new function names to defaults.
2. CopyPropertiesForUnrelatedObjects is called from RegenerateBlueprintClass, overwriting the new data with the defaults from the old CDO.
3. Compile derived anim BP (really just linking, as it is data-only), serializing incorrect old base defaults that were copied above.
4. Compile base bytecode as a result of FScopedDependencyGather (recompiles base, but derived is left with incorrect defaults as when it is compiled from FScopedDependencyGather - only a link is performed, no re-serialization).

Result = derived with incorrect properties that were copied from the old CDO in CopyPropertiesForUnrelatedObjects.

This fix prevents compiler generated defaults from being overwritten in CopyPropertiesForUnrelatedObjects. It also stabilizes the function names generated by the anim BP compiler (by using the node GUID instead of the FName). This prevents the compiled function names being clobbered and ensures that names match between runs of the compiler.

Also made sure that Guids that are duplicated because of copy-paste bug still provide unique function names.

UE-23436 - Little husks don't animate when running

diff --git a/Engine/Source/Editor/KismetCompiler/Private/AnimBlueprintCompiler.cpp b/Engine/Source/Editor/KismetCompiler/Private/AnimBlueprintCompiler.cpp
--- a/Engine/Source/Editor/KismetCompiler/Private/AnimBlueprintCompiler.cpp
+++ b/Engine/Source/Editor/KismetCompiler/Private/AnimBlueprintCompiler.cpp
@@ -131,9 +131,22 @@ void FAnimBlueprintCompiler::CreateEvaluationHandler(UAnimGraphNode_Base* Visual
 	check(Record.ServicedProperties.Num() > 0);
 	check(Record.NodeVariableProperty != NULL);
 
-	//@TODO: Want to name these better
-	const FString FunctionName = FString::Printf(TEXT("%s_%s_%s"), *Record.EvaluationHandlerProperty->GetName(), *VisualAnimNode->GetOuter()->GetName(), *VisualAnimNode->GetName());
+	// Use the node GUID for a stable name across compiles
+	FString FunctionName = FString::Printf(TEXT("%s_%s_%s_%s"), *Record.EvaluationHandlerProperty->GetName(), *VisualAnimNode->GetOuter()->GetName(), *VisualAnimNode->GetClass()->GetName(), *VisualAnimNode->NodeGuid.ToString());
 	Record.HandlerFunctionName = FName(*FunctionName);
+
+	// check function name isnt already used (data exists that can contain duplicate GUIDs) and apply a numeric extension until it is unique
+	int32 ExtensionIndex = 0;
+	FName* ExistingName = HandlerFunctionNames.Find(Record.HandlerFunctionName);
+	while(ExistingName != nullptr)
+	{
+		FunctionName = FString::Printf(TEXT("%s_%s_%s_%s_%d"), *Record.EvaluationHandlerProperty->GetName(), *VisualAnimNode->GetOuter()->GetName(), *VisualAnimNode->GetClass()->GetName(), *VisualAnimNode->NodeGuid.ToString(), ExtensionIndex);
+		Record.HandlerFunctionName = FName(*FunctionName);
+		ExistingName = HandlerFunctionNames.Find(Record.HandlerFunctionName);
+		ExtensionIndex++;
+	}
+
+	HandlerFunctionNames.Add(Record.HandlerFunctionName);
 	
 	// Add a custom event in the graph
 	UK2Node_CustomEvent* EntryNode = SpawnIntermediateNode<UK2Node_CustomEvent>(VisualAnimNode, ConsolidatedEventGraph);
diff --git a/Engine/Source/Editor/KismetCompiler/Private/AnimBlueprintCompiler.h b/Engine/Source/Editor/KismetCompiler/Private/AnimBlueprintCompiler.h
--- a/Engine/Source/Editor/KismetCompiler/Private/AnimBlueprintCompiler.h
+++ b/Engine/Source/Editor/KismetCompiler/Private/AnimBlueprintCompiler.h
@@ -215,6 +215,9 @@ protected:
 	// Map of cache name to encountered save cached pose nodes
 	TMap<FString, UAnimGraphNode_SaveCachedPose*> SaveCachedPoseNodes;
 
+	// Set of used handler function names
+	TSet<FName> HandlerFunctionNames;
+
 	// True if any parent class is also generated from an animation blueprint
 	bool bIsDerivedAnimBlueprint;
 private:
diff --git a/Engine/Source/Editor/UnrealEd/Private/Kismet2/BlueprintEditorUtils.cpp b/Engine/Source/Editor/UnrealEd/Private/Kismet2/BlueprintEditorUtils.cpp
--- a/Engine/Source/Editor/UnrealEd/Private/Kismet2/BlueprintEditorUtils.cpp
+++ b/Engine/Source/Editor/UnrealEd/Private/Kismet2/BlueprintEditorUtils.cpp
@@ -1672,6 +1672,7 @@ UClass* FBlueprintEditorUtils::RegenerateBlueprintClass(UBlueprint* Blueprint, U
 						CopyDetails.bAggressiveDefaultSubobjectReplacement = true;
 						CopyDetails.bDoDelta = false;
 						CopyDetails.bCopyDeprecatedProperties = true;
+						CopyDetails.bSkipCompilerGeneratedDefaults = true;
 						UEditorEngine::CopyPropertiesForUnrelatedObjects(PreviousCDO, NewCDO, CopyDetails);
 					}
 
diff --git a/Engine/Source/Runtime/CoreUObject/Public/UObject/ObjectBase.h b/Engine/Source/Runtime/CoreUObject/Public/UObject/ObjectBase.h
--- a/Engine/Source/Runtime/CoreUObject/Public/UObject/ObjectBase.h
+++ b/Engine/Source/Runtime/CoreUObject/Public/UObject/ObjectBase.h
@@ -1122,6 +1122,9 @@ namespace UM
 
 		/// [PropertyMetadata] Property is serialized to config and we should be able to set it anywhere along the config hierarchy.
 		ConfigHierarchyEditable,
+
+		/// [PropertyMetadata] Property defaults are generated by the Blueprint compiler and will not be copied when CopyPropertiesForUnrelatedObjects is called post-compile.
+		BlueprintCompilerGeneratedDefaults,
 	};
 
 	// Metadata usable in UPROPERTY for customizing the behavior of Persona and UMG
diff --git a/Engine/Source/Runtime/Engine/Classes/Animation/AnimNodeBase.h b/Engine/Source/Runtime/Engine/Classes/Animation/AnimNodeBase.h
--- a/Engine/Source/Runtime/Engine/Classes/Animation/AnimNodeBase.h
+++ b/Engine/Source/Runtime/Engine/Classes/Animation/AnimNodeBase.h
@@ -512,7 +512,7 @@ struct ENGINE_API FAnimNode_Base
 	GENERATED_USTRUCT_BODY()
 
 	// The default handler for graph-exposed inputs
-	UPROPERTY()
+	UPROPERTY(meta=(BlueprintCompilerGeneratedDefaults))
 	FExposedValueHandler EvaluateGraphExposedInputs;
 
 	// A derived class should implement Initialize, Update, and either Evaluate or EvaluateComponentSpace, but not both of them
diff --git a/Engine/Source/Runtime/Engine/Classes/Engine/Engine.h b/Engine/Source/Runtime/Engine/Classes/Engine/Engine.h
--- a/Engine/Source/Runtime/Engine/Classes/Engine/Engine.h
+++ b/Engine/Source/Runtime/Engine/Classes/Engine/Engine.h
@@ -2246,12 +2246,16 @@ public:
 		bool bCopyDeprecatedProperties;
 		bool bPreserveRootComponent;
 
+		/** Skips copying properties with BlueprintCompilerGeneratedDefaults metadata */
+		bool bSkipCompilerGeneratedDefaults;
+
 		FCopyPropertiesForUnrelatedObjectsParams()
 			: bAggressiveDefaultSubobjectReplacement(false)
 			, bDoDelta(true)
 			, bReplaceObjectClassReferences(true)
 			, bCopyDeprecatedProperties(false)
 			, bPreserveRootComponent(true)
+			, bSkipCompilerGeneratedDefaults(false)
 		{}
 	};
 	static void CopyPropertiesForUnrelatedObjects(UObject* OldObject, UObject* NewObject, FCopyPropertiesForUnrelatedObjectsParams Params = FCopyPropertiesForUnrelatedObjectsParams());//bool bAggressiveDefaultSubobjectReplacement = false, bool bDoDelta = true);
diff --git a/Engine/Source/Runtime/Engine/Private/Animation/AnimNodeBase.cpp b/Engine/Source/Runtime/Engine/Private/Animation/AnimNodeBase.cpp
--- a/Engine/Source/Runtime/Engine/Private/Animation/AnimNodeBase.cpp
+++ b/Engine/Source/Runtime/Engine/Private/Animation/AnimNodeBase.cpp
@@ -344,6 +344,7 @@ void FExposedValueHandler::Initialize(FAnimNode_Base* AnimNode, UObject* AnimIns
 	if (BoundFunction != NAME_None)
 	{
 		Function = AnimInstanceObject->FindFunction(BoundFunction);
+		check(Function);
 	}
 	else
 	{
diff --git a/Engine/Source/Runtime/Engine/Private/UnrealEngine.cpp b/Engine/Source/Runtime/Engine/Private/UnrealEngine.cpp
--- a/Engine/Source/Runtime/Engine/Private/UnrealEngine.cpp
+++ b/Engine/Source/Runtime/Engine/Private/UnrealEngine.cpp
@@ -10817,7 +10817,33 @@ void UEngine::CopyPropertiesForUnrelatedObjects(UObject* OldObject, UObject* New
 	const uint32 AdditinalPortFlags = Params.bCopyDeprecatedProperties ? PPF_UseDeprecatedProperties : PPF_None;
 	// Save the modified properties of the old CDO
 	{
-		FObjectWriter Writer(OldObject, SavedProperties, true, true, Params.bDoDelta, AdditinalPortFlags);
+		class FCopyPropertiesArchiveObjectWriter : public FObjectWriter
+		{
+		public:
+			FCopyPropertiesArchiveObjectWriter(UObject* Obj, TArray<uint8>& InBytes, bool bIgnoreClassRef, bool bIgnoreArchetypeRef, bool bDoDelta , uint32 AdditionalPortFlags, bool bInSkipCompilerGeneratedDefaults)
+				: FObjectWriter(InBytes)
+			{	
+				bSkipCompilerGeneratedDefaults = bInSkipCompilerGeneratedDefaults;
+				ArIgnoreClassRef = bIgnoreClassRef;
+				ArIgnoreArchetypeRef = bIgnoreArchetypeRef;
+				ArNoDelta = !bDoDelta;
+				ArPortFlags |= AdditionalPortFlags;
+
+				Obj->Serialize(*this);
+			}
+
+#if WITH_EDITOR
+			virtual bool ShouldSkipProperty(const class UProperty* InProperty) const override
+			{
+				static FName BlueprintCompilerGeneratedDefaultsName(TEXT("BlueprintCompilerGeneratedDefaults"));
+				return bSkipCompilerGeneratedDefaults && InProperty->HasMetaData(BlueprintCompilerGeneratedDefaultsName);
+			}
+#endif
+
+			bool bSkipCompilerGeneratedDefaults;
+		};
+
+		FCopyPropertiesArchiveObjectWriter Writer(OldObject, SavedProperties, true, true, Params.bDoDelta, AdditinalPortFlags, Params.bSkipCompilerGeneratedDefaults);
 	}
 
 	{
