Copying //UE4/Dev-Blueprints to Dev-Main (//UE4/Dev-Main) @ 2781164

diff --git a/Engine/Source/Editor/Kismet/Private/Profiler/ScriptEventExecution.cpp b/Engine/Source/Editor/Kismet/Private/Profiler/ScriptEventExecution.cpp
new file mode 100644
--- /dev/null
+++ b/Engine/Source/Editor/Kismet/Private/Profiler/ScriptEventExecution.cpp
@@ -0,0 +1,487 @@
+// Copyright 1998-2015 Epic Games, Inc. All Rights Reserved.
+
+#include "BlueprintEditorPrivatePCH.h"
+#include "Developer/BlueprintProfiler/Public/BlueprintProfilerModule.h"
+#include "ScriptEventExecution.h"
+#include "BPProfilerStatisticWidgets.h"
+
+//////////////////////////////////////////////////////////////////////////
+// FScriptExecEvent
+
+FScriptExecEvent::FScriptExecEvent()
+	: EventTime(0.0)
+{
+}
+
+FScriptExecEvent::FScriptExecEvent(TWeakObjectPtr<const UObject> InObjectPtr, const FScriptExecEvent& ContextToCopy)
+	: EventType(EScriptInstrumentation::NodePin)
+	, ScriptCodeOffset(-1)
+	, EventTime(0.0)
+	, ContextObject(InObjectPtr)
+	, PinIndex(ContextToCopy.PinIndex)
+	, CallDepth(ContextToCopy.CallDepth)
+	, NumExecOutPins(1)
+{
+}
+
+FScriptExecEvent::FScriptExecEvent(TWeakObjectPtr<const UObject> InObjectPtr, TWeakObjectPtr<const UObject> InMacroObjectPtr, const FBlueprintInstrumentedEvent& InEvent, const int32 InPinIndex)
+	: EventType(InEvent.GetType())
+	, ScriptCodeOffset(InEvent.GetScriptCodeOffset())
+	, EventTime(InEvent.GetTime())
+	, ContextObject(InObjectPtr)
+	, MacroContextObject(InMacroObjectPtr)
+	, PinIndex(InPinIndex)
+	, CallDepth(0)
+	, NumExecOutPins(-1)
+{
+}
+
+bool FScriptExecEvent::operator == (const FScriptExecEvent& ContextIn) const
+{
+	return EventType == ContextIn.EventType && ContextObject == ContextIn.ContextObject;
+}
+
+bool FScriptExecEvent::CanSubmitData() const
+{
+	return	EventType != EScriptInstrumentation::Class && 
+			EventType != EScriptInstrumentation::Instance;
+}
+
+bool FScriptExecEvent::IsBranch()
+{
+	if (NumExecOutPins < 0)
+	{
+		NumExecOutPins = 0;
+		// Cache num exec pins
+		if (const UEdGraphNode* GraphNode = Cast<UEdGraphNode>(ContextObject.Get())) 
+		{
+			NumExecOutPins = 0;
+			for (auto Pin : GraphNode->Pins)
+			{
+				if (Pin->Direction == EGPD_Output && Pin->PinType.PinCategory == UEdGraphSchema_K2::PC_Exec)
+				{
+					NumExecOutPins++;
+				}
+			}
+		}
+		else
+		{
+			// for non node contexts just flag that we tried to discover the pins.
+			NumExecOutPins = 1;
+		}
+	}
+	return NumExecOutPins > 1;
+}
+
+TWeakObjectPtr<const UObject> FScriptExecEvent::GetGraphPinPtr()
+{
+	TWeakObjectPtr<const UObject> Result;
+	if (PinIndex >= 0)
+	{
+		if (const UEdGraphNode* Node = Cast<const UEdGraphNode>(ContextObject.Get()))
+		{
+			Result = Node->Pins[PinIndex];
+		}
+	}
+	return Result;
+}
+
+void FScriptExecEvent::MergeContext(const FScriptExecEvent& TerminatorContext)
+{
+	EventTime = TerminatorContext.GetTime() - EventTime;
+	PinIndex = TerminatorContext.PinIndex;
+}
+
+//////////////////////////////////////////////////////////////////////////
+// FProfilerClassContext
+
+void FProfilerClassContext::ProcessExecutionPath(const TArray<FBlueprintInstrumentedEvent>& ProfilerData, const int32 ExecStart, const int32 ExecEnd)
+{
+	// Process profiling events for a single execution path
+	for (int32 Idx = ExecStart; Idx <= ExecEnd; ++Idx)
+	{
+		// Cursory check for end submission conditions
+		switch(ProfilerData[Idx].GetType())
+		{
+			case EScriptInstrumentation::Class:
+			{
+				ContextStack.SetNum(0, false);
+				ContextStack.Add(GenerateContextFromEvent(ProfilerData[Idx]));
+				// Ensure the root stat exists for this blueprint.
+				if (!ClassStat.IsValid())
+				{
+					CreateClassStat();
+				}
+				break;
+			}
+			case EScriptInstrumentation::Instance:
+			{
+				ContextStack.SetNum(1, false);
+				ContextStack.Add(GenerateContextFromEvent(ProfilerData[Idx]));
+				// Ensure the root stat exists for this blueprint.
+				if (!ClassStat.IsValid())
+				{
+					CreateClassStat();
+				}
+				break;
+			}
+			case EScriptInstrumentation::Event:
+			{
+				ContextStack.SetNum(2, false);
+				ContextStack.Add(GenerateContextFromEvent(ProfilerData[Idx]));
+				break;
+			}
+			case EScriptInstrumentation::Stop:
+			{
+				ContextStack.Add(GenerateContextFromEvent(ProfilerData[Idx]));
+				FixupEventId();
+				RemoveDuplicateTraces();
+				OptimiseEvents();
+				SubmitData();
+				return;
+			}
+			default:
+			{
+				// Generate the new execution event
+				ContextStack.Add(GenerateContextFromEvent(ProfilerData[Idx]));
+				break;
+			}
+		}
+	}
+}
+
+void FProfilerClassContext::FixupEventId()
+{
+	FScriptExecEvent* EventContext = nullptr;
+	FScriptExecEvent* NodeContext = nullptr;
+	for (int32 Idx = 0; Idx < ContextStack.Num(); ++Idx)
+	{
+		FScriptExecEvent& CurrContext = ContextStack[Idx];
+		if (!EventContext && CurrContext.IsEvent())
+		{
+			EventContext = &CurrContext;
+		}
+		else if (!NodeContext && CurrContext.IsNodeEntry())
+		{
+			NodeContext = &CurrContext;
+			break;
+		}
+	}
+	if (EventContext && NodeContext)
+	{
+		EventContext->SetObjectPtr(NodeContext->GetObjectPtr());
+	}
+}
+
+void FProfilerClassContext::RemoveDuplicateTraces()
+{
+	// Remove duplicate traces and reset call depths
+	for (int32 Idx = 0; (Idx+1) < ContextStack.Num(); ++Idx)
+	{
+		ContextStack[Idx].ResetCallDepth();
+		if (ContextStack[Idx] == ContextStack[Idx+1])
+		{
+			ContextStack.RemoveAt(Idx, 1, false);
+			Idx--;
+		}
+	}
+}
+
+void FProfilerClassContext::OptimiseEvents()
+{
+	// flatten call graphs updating depth
+	for (int32 Idx = 0; Idx < ContextStack.Num(); ++Idx)
+	{
+		switch(ContextStack[Idx].GetType())
+		{
+			case EScriptInstrumentation::Class:
+			case EScriptInstrumentation::Instance:
+			{
+				// Artificial events that provide a const offset for all node timings.
+				for (int32 ChildIdx = Idx + 1; ChildIdx < ContextStack.Num(); ++ChildIdx)
+				{
+					ContextStack[ChildIdx].IncrementDepth();
+				}
+				break;
+			}
+			case EScriptInstrumentation::NodeExit:
+			{
+				TWeakObjectPtr<const UObject> ObjKey = ContextStack[Idx].GetObjectPtr();
+				int32 PairedEntryIdx = Idx-1;
+				// Find entry point
+				for (; PairedEntryIdx >= 0; --PairedEntryIdx)
+				{
+					if (ContextStack[PairedEntryIdx].GetObjectPtr() == ObjKey)
+					{
+						break;
+					}
+					else
+					{
+						ContextStack[PairedEntryIdx].IncrementDepth();
+					}
+				}
+				// Merge paired nodes and remove the exit node.
+				if (PairedEntryIdx > 0)
+				{
+					ContextStack[PairedEntryIdx].MergeContext(ContextStack[Idx]);
+					ContextStack.RemoveAt(Idx, 1, false);
+					Idx--;
+				}
+				break;
+			}
+			case EScriptInstrumentation::Stop:
+			{
+				int32 PairedEntryIdx = Idx-1;
+				// Find entry point
+				for (; PairedEntryIdx >= 0; --PairedEntryIdx)
+				{
+					if (ContextStack[PairedEntryIdx].IsEvent())
+					{
+						break;
+					}
+					else
+					{
+						ContextStack[PairedEntryIdx].IncrementDepth();
+					}
+				}
+				// Merge paired nodes and remove the exit node.
+				if (PairedEntryIdx > 0)
+				{
+					ContextStack[PairedEntryIdx].MergeContext(ContextStack[Idx]);
+					ContextStack.RemoveAt(Idx, 1, false);
+					Idx--;
+				}
+				break;
+			}
+		}
+	}
+	// Interpret graph node event types, and handle branch depths
+	for (int32 Idx = 0; Idx < ContextStack.Num(); ++Idx)
+	{
+		if (ContextStack[Idx].IsNodeEntry())
+		{
+			if (ContextStack[Idx].IsMacro())
+			{
+				// Collapse all macro nodes to just the instance for now.
+				// Further work will include ensuring the BPGC has valid pins to lookup for macros.
+				ContextStack[Idx].SetType(EScriptInstrumentation::Macro);
+				int32 NumToRemove = 0;
+				double MacroTiming = 0.0;
+				// Merge node timings
+				for (int32 MacroIdx = Idx; MacroIdx < ContextStack.Num(); ++MacroIdx)
+				{
+					if (ContextStack[Idx].GetMacroObjectPtr() != ContextStack[MacroIdx].GetMacroObjectPtr())
+					{
+						break;
+					}
+					NumToRemove++;
+					MacroTiming += ContextStack[MacroIdx].GetTime();
+				}
+				if (NumToRemove > 1)
+				{
+					ContextStack[Idx].SetTime(MacroTiming);
+					ContextStack.RemoveAt(Idx+1, NumToRemove, false);
+				}
+			}
+			else if (ContextStack[Idx].IsFunctionCallSite())
+			{
+				ContextStack[Idx].SetType(EScriptInstrumentation::Function);
+			}
+			else if (ContextStack[Idx].IsBranch())
+			{
+				ContextStack[Idx].SetType(EScriptInstrumentation::Branch);
+				const int32 BranchDepth = ContextStack[Idx].GetCallDepth();
+				TWeakObjectPtr<const UObject> BranchNodePtr = ContextStack[Idx].GetObjectPtr();
+				int32 NodeDepth = BranchDepth;
+				// Decrement branch children depths
+				for (int32 BranchChildIdx = Idx + 1; BranchChildIdx < ContextStack.Num() && ContextStack[BranchChildIdx].GetCallDepth() >= BranchDepth; ++BranchChildIdx)
+				{
+					if (ContextStack[BranchChildIdx].GetObjectPtr() != BranchNodePtr)
+					{
+						ContextStack[BranchChildIdx].IncrementDepth();
+					}
+					else
+					{
+						break;
+					}
+				}
+			}
+		}
+	}
+}
+
+void FProfilerClassContext::CreateClassStat()
+{
+	if (!ClassStat.IsValid())
+	{
+		// Determine class from class or instance event.
+		for (int32 Idx = 0; Idx < ContextStack.Num(); ++Idx)
+		{
+			if (ContextStack[Idx].IsClass())
+			{
+				FBPProfilerStatPtr Root;
+				ClassStat = MakeShareable(new FBPProfilerStat(Root, ContextStack[Idx]));
+				break;
+			}
+			else if (ContextStack[Idx].IsInstance())
+			{
+				if (const UObject* InstanceObject = ContextStack[Idx].GetObjectPtr().Get())
+				{
+					if (UClass* InstanceClass = InstanceObject->GetClass())
+					{
+						FBlueprintInstrumentedEvent FakeClassEvent(EScriptInstrumentation::Class, InstanceClass->GetPathName(), -1);
+						FScriptExecEvent FakeExecContext = GenerateContextFromEvent(FakeClassEvent);
+						FBPProfilerStatPtr Root;
+						ClassStat = MakeShareable(new FBPProfilerStat(Root, FakeExecContext));
+						break;
+					}
+				}
+			}
+		}
+	}
+}
+
+void FProfilerClassContext::SubmitData()
+{
+	check(ClassStat.IsValid());
+	if (ContextStack.Num())
+	{
+		// Submit profiling data now it is fully conformed.
+		FBPProfilerStatPtr CurrentParentStat = ClassStat->FindEntryPoint(ContextStack[0]);
+		if (CurrentParentStat.IsValid())
+		{
+			CurrentParentStat->SubmitData(ContextStack[0]);
+			int32 CurrentCallDepth = ContextStack[0].GetCallDepth() + 1;
+			FBPProfilerStatPtr CurrentStat;
+			// Debug Extras
+			FScriptExecEvent ParentContext = CurrentParentStat->GetObjectContext();
+
+			for (int32 Idx = 1; Idx < ContextStack.Num(); ++Idx)
+			{
+				FScriptExecEvent& NextContext = ContextStack[Idx];
+				// Adjust call depth if required
+				if (NextContext.GetCallDepth() > CurrentCallDepth)
+				{
+					CurrentParentStat = CurrentStat;
+					ParentContext = CurrentParentStat->GetObjectContext();
+					CurrentCallDepth++;
+				}
+				else if(NextContext.GetCallDepth() < CurrentCallDepth)
+				{
+					while(NextContext.GetCallDepth() < CurrentCallDepth)
+					{
+						CurrentParentStat = CurrentParentStat->GetParent();
+						ParentContext = CurrentParentStat->GetObjectContext();
+						CurrentCallDepth = CurrentParentStat->GetCallDepth() + 1;
+					}
+				}
+				// Submit Data
+				const bool bBranchNode = NextContext.IsBranch();
+				CurrentStat = CurrentParentStat->FindOrAddChildByContext(NextContext);
+				check (CurrentStat.IsValid());
+				CurrentStat->SubmitData(NextContext);
+			}
+			// Update stat data tree.
+			ClassStat->UpdateProfilingStats();
+		}
+	}
+}
+
+FScriptExecEvent FProfilerClassContext::GenerateContextFromEvent(const FBlueprintInstrumentedEvent& Event)
+{
+	// Try the cache first, using the object path and script offset.
+	const FString ObjectPath = FString::Printf( TEXT( "%s:%i" ), *Event.GetObjectPath(), Event.GetScriptCodeOffset());
+	FEventObjectContext& Result = PathToObjectCache.FindOrAdd(ObjectPath);
+	int32 DiscoveredPinIndex = -1;
+	// Fall back to more exhaustive search, but this should only happen once per object.
+	if (!Result.IsValid())
+	{
+		if (Event.IsNewClass())
+		{
+			// Attempt to locate the blueprint
+			if (const UObject* ObjectPtr = FindObject<UObject>(nullptr, *Event.GetObjectPath()))
+			{
+				if (const UBlueprintGeneratedClass* BPGC = Cast<UBlueprintGeneratedClass>(ObjectPtr))
+				{
+					if (const UBlueprint* Blueprint = Cast<UBlueprint>(BPGC->ClassGeneratedBy))
+					{
+						Result.ContextObject = Blueprint;
+					}
+				}
+			}
+		}
+		else if (Event.IsNewInstance())
+		{
+			// attempt to locate the instance
+			if (const UObject* ObjectPtr = FindObject<UObject>(nullptr, *Event.GetObjectPath()))
+			{
+				// Get Outer world
+				if (UWorld* ObjectWorld = ObjectPtr->GetTypedOuter<UWorld>())
+				{
+					if (ObjectWorld->WorldType == EWorldType::PIE)
+					{
+						FWorldContext& EditorWorldContext = GEditor->GetEditorWorldContext();
+						if (UWorld* EditorWorld = EditorWorldContext.World())
+						{
+							for (auto LevelIter : EditorWorld->GetLevels())
+							{
+								if (UObject* EditorObject = FindObject<UObject>(LevelIter, *ObjectPtr->GetName()))
+								{
+									Result.ContextObject = EditorObject;
+									break;
+								}
+							}
+						}
+					}
+					else
+					{
+						Result.ContextObject = ObjectPtr;
+					}
+				}
+			}
+		}
+		else
+		{
+			// Attempt to locate graph node
+			TArray<FString> ObjectPaths;
+			Event.GetObjectPath().ParseIntoArray(ObjectPaths, TEXT( ":" ));
+
+			if (ObjectPaths.Num() > 0)
+			{
+				if (UBlueprintGeneratedClass* BlueprintClass = FindObject<UBlueprintGeneratedClass>(nullptr, *ObjectPaths[0]))
+				{
+					Result.BlueprintClass = BlueprintClass;
+					if (UFunction* BlueprintFunction = BlueprintClass->FindFunctionByName(FName(*ObjectPaths.Last())))
+					{
+						Result.BlueprintFunction = BlueprintFunction;
+						if (const UEdGraphNode* GraphNode = BlueprintClass->GetDebugData().FindSourceNodeFromCodeLocation(BlueprintFunction, Event.GetScriptCodeOffset(), true))
+						{
+							TArray<UEdGraphNode*> MacroInstanceNodes;
+							BlueprintClass->GetDebugData().FindMacroInstanceNodesFromCodeLocation(BlueprintFunction, Event.GetScriptCodeOffset(), MacroInstanceNodes);
+							// We need to do more work to ensure the executable path in macros can be followed so disabling for now.
+							if (MacroInstanceNodes.Num())
+							{
+								Result.ContextObject = BlueprintClass->GetDebugData().FindMacroSourceNodeFromCodeLocation(BlueprintFunction, Event.GetScriptCodeOffset());
+								Result.MacroContextObject = MacroInstanceNodes[0];
+							}
+							else
+							{
+								Result.ContextObject = GraphNode;
+							}
+						}
+					}
+				}
+			}
+		}
+	}
+	// Find exit pin index
+	if (Result.DoesRequireExecPin())
+	{
+		if (UEdGraphPin* GraphPin = Result.BlueprintClass.Get()->GetDebugData().FindExecPinFromCodeLocation(Result.BlueprintFunction.Get(), Event.GetScriptCodeOffset()))
+		{
+			DiscoveredPinIndex = Result.GetGraphNode()->GetPinIndex(GraphPin);
+		}
+	}
+	FScriptExecEvent NewContext(Result.ContextObject, Result.MacroContextObject, Event, DiscoveredPinIndex);
+	return NewContext;
+}
diff --git a/Engine/Source/Editor/Kismet/Private/Profiler/ScriptEventExecution.h b/Engine/Source/Editor/Kismet/Private/Profiler/ScriptEventExecution.h
new file mode 100644
--- /dev/null
+++ b/Engine/Source/Editor/Kismet/Private/Profiler/ScriptEventExecution.h
@@ -0,0 +1,169 @@
+// Copyright 1998-2015 Epic Games, Inc. All Rights Reserved.
+
+#pragma once
+
+typedef TSharedPtr<class FBPProfilerStat> FBPProfilerStatPtr;
+typedef TMap<TWeakObjectPtr<const UObject>, FBPProfilerStatPtr> BPProfilerWidgetCache;
+
+//////////////////////////////////////////////////////////////////////////
+// FScriptExecEvent
+
+class FScriptExecEvent
+{
+public:
+
+	FScriptExecEvent();
+	FScriptExecEvent(TWeakObjectPtr<const UObject> InObjectPtr, const FScriptExecEvent& ContextToCopy);
+	FScriptExecEvent(TWeakObjectPtr<const UObject> InObjectPtr, TWeakObjectPtr<const UObject> InMacroObjectPtr, const class FBlueprintInstrumentedEvent& InEvent, const int32 InPinIndex);
+
+	bool operator == (const FScriptExecEvent& ContextIn) const;
+
+	bool IsValid() const { return ContextObject.IsValid(); }
+
+	/** Override the event object */
+	void SetObjectPtr(TWeakObjectPtr<const UObject> InObjPtr) { ContextObject = InObjPtr; }
+
+	/** Returns the event object */
+	TWeakObjectPtr<const UObject> GetObjectPtr() { return ContextObject; }
+
+	/** Returns the macro event object */
+	TWeakObjectPtr<const UObject> GetMacroObjectPtr() { return MacroContextObject; }
+
+	/** Returns the context object as the template type */
+	template<typename UserType> TWeakObjectPtr<const UserType> GetTypedObjectPtr()
+	{
+		return TWeakObjectPtr<const UserType>(Cast<const UserType>(ContextObject.Get()));
+	}
+
+	/** Returns the graph pin the execution path used */
+	TWeakObjectPtr<const UObject> GetGraphPinPtr();
+
+	/** Returns the profiling event type */
+	const EScriptInstrumentation::Type GetType() const { return EventType; }
+
+	/** Overrides the profiling event type */
+	void SetType(const EScriptInstrumentation::Type NewEventType) { EventType = NewEventType; }
+
+	/** Returns the execution duration for this execution event */
+	double GetTime() const { return EventTime; }
+
+	/** Overrides the timing stat */
+	void SetTime(double NewTime) { EventTime = NewTime; }
+
+	/** Returns the call depth of this execution event */
+	int32 GetCallDepth() const { return CallDepth; }
+
+	/** Resets the call depth of this execution event */
+	void ResetCallDepth() { CallDepth = 0; }
+
+	/** Increments the call depth of this execution event */
+	void IncrementDepth() { CallDepth++; }
+
+	/** Merge start/end node events and calculate the elapsed time */
+	void MergeContext(const FScriptExecEvent& TerminatorContext);
+
+	/** Returns if this exec context requires event timings */
+	bool CanSubmitData() const;
+
+	/** Returns if this exec event represents a change in class/blueprint */
+	bool IsClass() const { return EventType == EScriptInstrumentation::Class; }
+
+	/** Returns if this exec event represents a change in instance */
+	bool IsInstance() const { return EventType == EScriptInstrumentation::Instance; }
+
+	/** Returns if this exec event represents the start of an event execution path */
+	bool IsEvent() const { return EventType == EScriptInstrumentation::Event; }
+
+	/** Returns if this event is a function callsite event */
+	bool IsFunctionCallSite() const { return IsValid() ? ContextObject.Get()->IsA<UK2Node_CallFunction>() : false; }
+
+	/** Returns if this event happened inside a macro instance */
+	bool IsMacro() const { return MacroContextObject.IsValid(); }
+
+	/** Returns if this event potentially multiple exit sites */
+	bool IsBranch();
+
+	/** Returns if this event signals the start of a node timing. */
+	bool IsNodeEntry() const { return EventType == EScriptInstrumentation::NodeEntry; }
+
+	/** Returns if this event signals the end of a node timing. */
+	bool IsNodeExit() const { return EventType == EScriptInstrumentation::NodeExit; }
+
+protected:
+
+	/** Calculated call depth */
+	int32 CallDepth;
+	/** Discovered exec out pin count */
+	int32 NumExecOutPins;
+	/** Exec event type */
+	EScriptInstrumentation::Type EventType;
+	/** Script code offset at event */
+	int32 ScriptCodeOffset;
+	/** System time event occurred */
+	double EventTime;
+	/** Context object */
+	TWeakObjectPtr<const UObject> ContextObject;
+	/** Macro Context object */
+	TWeakObjectPtr<const UObject> MacroContextObject;
+	/** Exit pin index */
+	int32 PinIndex;
+
+};
+
+//////////////////////////////////////////////////////////////////////////
+// FEventObjectContext
+
+struct FEventObjectContext
+{
+	bool IsValid() const { return ContextObject.IsValid(); }
+	bool DoesRequireExecPin() const { return BlueprintClass.IsValid() && BlueprintClass.IsValid(); }
+	const class UEdGraphNode* GetGraphNode() const { return ContextObject.IsValid() ? Cast<UEdGraphNode>(ContextObject.Get()) : nullptr; }
+
+	TWeakObjectPtr<const UObject> ContextObject;
+	TWeakObjectPtr<const UObject> MacroContextObject;
+	TWeakObjectPtr<UBlueprintGeneratedClass> BlueprintClass;
+	TWeakObjectPtr<UFunction> BlueprintFunction;
+};
+
+//////////////////////////////////////////////////////////////////////////
+// FProfilerClassContext
+
+class FProfilerClassContext
+{
+public:
+
+	/** Returns that widget that represents the class blueprint */
+	FBPProfilerStatPtr GetRootStat() { return ClassStat; }
+
+	/** This analyses, conforms and submits profiling data */
+	void ProcessExecutionPath(const TArray<FBlueprintInstrumentedEvent>& ProfilerData, const int32 ExecStart, const int32 ExecEnd);
+
+private:
+
+	/** Creates the root statisic for this class */
+	void CreateClassStat();
+
+	/** Use the current event stack to populate the event marker with a valid object */
+	void FixupEventId();
+
+	/** Remove duplicate entry profiling markers so the data is conformed enough to process */
+	void RemoveDuplicateTraces();
+
+	/** Merge node exit/entries into single events, Sort call depths of profiling events */
+	void OptimiseEvents();
+
+	/** Submit the conformed profiling data to the widget structure for display */
+	void SubmitData();
+
+	/** Create FScriptExecEvent from profiling event data and locate and cache the contextual UObject */
+	FScriptExecEvent GenerateContextFromEvent(const FBlueprintInstrumentedEvent& InEvent);
+
+private:
+
+	/** Root blueprint stat, representing the root of all execution flows in this class */
+	FBPProfilerStatPtr ClassStat;
+	/** Stack of processed events ready for submission */
+	TArray<FScriptExecEvent> ContextStack;
+	/** Cached object path and code offset lookup to UObjects */
+	TMap<const FString, struct FEventObjectContext> PathToObjectCache;
+};
diff --git a/Engine/Source/Editor/Kismet/Private/SBlueprintEditorToolbar.cpp b/Engine/Source/Editor/Kismet/Private/SBlueprintEditorToolbar.cpp
--- a/Engine/Source/Editor/Kismet/Private/SBlueprintEditorToolbar.cpp
+++ b/Engine/Source/Editor/Kismet/Private/SBlueprintEditorToolbar.cpp
@@ -175,6 +175,15 @@ void FKismet2Menu::FillDebugMenu( FMenuBuilder& MenuBuilder )
 	MenuBuilder.EndSection();
 }
 
+void FKismet2Menu::FillProfilerMenu( FMenuBuilder& MenuBuilder )
+{
+	MenuBuilder.BeginSection("Profiler", LOCTEXT("ProfilerMenu_Heading", "Profiler") );
+	{
+		MenuBuilder.AddMenuEntry( FFullBlueprintEditorCommands::Get().ToggleProfiler );
+	}
+	MenuBuilder.EndSection();
+}
+
 void FKismet2Menu::SetupBlueprintEditorMenu( TSharedPtr< FExtender > Extender, FBlueprintEditor& BlueprintEditor)
 {
 	// Extend the File menu with asset actions
@@ -210,6 +219,16 @@ void FKismet2Menu::SetupBlueprintEditorMenu( TSharedPtr< FExtender > Extender, F
 					LOCTEXT("DebugMenu_ToolTip", "Open the debug menu"),
 					FNewMenuDelegate::CreateStatic( &FKismet2Menu::FillDebugMenu ),
 					"Debug");
+
+				if (GetDefault<UEditorExperimentalSettings>()->bBlueprintPerformanceAnalysisTools)
+				{
+					// Profiler
+					MenuBarBuilder.AddPullDownMenu( 
+						LOCTEXT("ProfilerMenu", "Profiler"),
+						LOCTEXT("ProfilerMenu_ToolTip", "Toggle the Profiler Menu"),
+						FNewMenuDelegate::CreateStatic( &FKismet2Menu::FillProfilerMenu ),
+						"Profiler");
+				}
 			}
 		};
 
@@ -344,6 +363,8 @@ void FFullBlueprintEditorCommands::RegisterCommands()
 	UI_COMMAND(SwitchToBlueprintDefaultsMode, "Defaults", "Switches to Class Defaults Mode", EUserInterfaceActionType::ToggleButton, FInputChord());
 	UI_COMMAND(SwitchToComponentsMode, "Components", "Switches to Components Mode", EUserInterfaceActionType::ToggleButton, FInputChord());
 
+	UI_COMMAND(ToggleProfiler,"Profiler", "Toggle Profiler State", EUserInterfaceActionType::ToggleButton, FInputChord());
+
 	UI_COMMAND(EditGlobalOptions, "Class Settings", "Edit Class Settings (Previously known as Blueprint Props)", EUserInterfaceActionType::ToggleButton, FInputChord());
 	UI_COMMAND(EditClassDefaults, "Class Defaults", "Edit the initial values of your class.", EUserInterfaceActionType::ToggleButton, FInputChord());
 
@@ -480,6 +501,17 @@ void FBlueprintEditorToolbar::AddComponentsToolbar(TSharedPtr<FExtender> Extende
 		FToolBarExtensionDelegate::CreateSP( this, &FBlueprintEditorToolbar::FillComponentsToolbar ) );
 }
 
+void FBlueprintEditorToolbar::AddProfilerToolbar(TSharedPtr<FExtender> Extender)
+{
+	TSharedPtr<FBlueprintEditor> BlueprintEditorPtr = BlueprintEditor.Pin();
+
+	Extender->AddToolBarExtension(
+		"Asset",
+		EExtensionHook::After,
+		BlueprintEditorPtr->GetToolkitCommands(),
+		FToolBarExtensionDelegate::CreateSP( this, &FBlueprintEditorToolbar::FillProfilerToolbar ) );
+}
+
 void FBlueprintEditorToolbar::FillBlueprintEditorModesToolbar(FToolBarBuilder& ToolbarBuilder)
 {
 	TSharedPtr<FBlueprintEditor> BlueprintEditorPtr = BlueprintEditor.Pin();
@@ -546,6 +578,29 @@ void FBlueprintEditorToolbar::FillCompileToolbar(FToolBarBuilder& ToolbarBuilder
 	ToolbarBuilder.EndSection();
 }
 
+void FBlueprintEditorToolbar::FillProfilerToolbar(FToolBarBuilder& ToolbarBuilder)
+{
+	const FFullBlueprintEditorCommands& Commands = FFullBlueprintEditorCommands::Get();
+	TSharedPtr<FBlueprintEditor> BlueprintEditorPtr = BlueprintEditor.Pin();
+	UBlueprint* BlueprintObj = BlueprintEditorPtr->GetBlueprintObj();
+	
+	ToolbarBuilder.BeginSection("Profiler");
+	
+	if(BlueprintObj != NULL)
+	{
+		ToolbarBuilder.AddToolBarButton( Commands.ToggleProfiler,
+										 NAME_None, 
+										 TAttribute<FText>(),
+										 TAttribute<FText>(this, &FBlueprintEditorToolbar::GetProfilerStatusTooltip),
+										 FSlateIcon(FEditorStyle::GetStyleSetName(), "BlueprintEditor.EnableProfiling"),
+										 FName(TEXT("ToggleProfiler")));
+
+
+	}
+	
+	ToolbarBuilder.EndSection();
+}
+
 void FBlueprintEditorToolbar::FillNewToolbar(FToolBarBuilder& ToolbarBuilder)
 {
 	const FBlueprintEditorCommands& Commands = FBlueprintEditorCommands::Get();
@@ -668,6 +723,12 @@ FText FBlueprintEditorToolbar::GetStatusTooltip() const
 	}
 }
 
+FText FBlueprintEditorToolbar::GetProfilerStatusTooltip() const
+{
+	return BlueprintEditor.Pin()->IsProfilerActive() ? LOCTEXT("Profiling_Enabled", "Profiling Active") :
+														LOCTEXT("Profiling_Disabled", "Profiling Inactive");
+}
+
 TArray< TSharedPtr< SWidget> > FBlueprintEditorToolbar::GenerateToolbarWidgets(const UBlueprint* BlueprintObj, TAttribute<FName> ActiveModeGetter, FOnModeChangeRequested ActiveModeSetter)
 {
 	TArray< TSharedPtr< SWidget> > Ret;
diff --git a/Engine/Source/Editor/Kismet/Private/SMyBlueprint.cpp b/Engine/Source/Editor/Kismet/Private/SMyBlueprint.cpp
--- a/Engine/Source/Editor/Kismet/Private/SMyBlueprint.cpp
+++ b/Engine/Source/Editor/Kismet/Private/SMyBlueprint.cpp
@@ -1218,7 +1218,8 @@ void SMyBlueprint::CollectAllActions(FGraphActionListBuilderBase& OutAllActions)
 	TArray<FName> OverridableFunctionNames;
 
 	// Cache potentially overridable functions
-	for ( TFieldIterator<UFunction> FunctionIt(BlueprintObj->ParentClass, EFieldIteratorFlags::IncludeSuper); FunctionIt; ++FunctionIt )
+	UClass* ParentClass = BlueprintObj->SkeletonGeneratedClass? BlueprintObj->SkeletonGeneratedClass->GetSuperClass() : BlueprintObj->ParentClass;
+	for ( TFieldIterator<UFunction> FunctionIt(ParentClass, EFieldIteratorFlags::IncludeSuper); FunctionIt; ++FunctionIt )
 	{
 		const UFunction* Function = *FunctionIt;
 		const FName FunctionName = Function->GetFName();
@@ -1226,7 +1227,7 @@ void SMyBlueprint::CollectAllActions(FGraphActionListBuilderBase& OutAllActions)
 		if (    UEdGraphSchema_K2::CanKismetOverrideFunction(Function) 
 			 && !OverridableFunctionNames.Contains(FunctionName) 
 			 && !ImplementedFunctionCache.Contains(FunctionName) 
-			 && !FObjectEditorUtils::IsFunctionHiddenFromClass(Function, BlueprintObj->ParentClass)
+			 && !FObjectEditorUtils::IsFunctionHiddenFromClass(Function, ParentClass)
 			 && !FBlueprintEditorUtils::FindOverrideForFunction(BlueprintObj, CastChecked<UClass>(Function->GetOuter()), Function->GetFName()) )
 		{
 			FString FunctionTooltip = UK2Node_CallFunction::GetDefaultTooltipForFunction(Function);
@@ -2119,19 +2120,19 @@ void SMyBlueprint::OnFindReference()
 		if (Guid.IsValid())
 		{
 			FString VariableName = VarAction->GetVariableName().ToString();
-			SearchTerm = FString::Printf(TEXT("Nodes(VariableReference(MemberName=+%s && MemberGuid(A=%i && B=%i && C=%i && D=%i)))"), *VariableName, Guid.A, Guid.B, Guid.C, Guid.D);
+			SearchTerm = FString::Printf(TEXT("Nodes(VariableReference(MemberName=+\"%s\" && MemberGuid(A=%i && B=%i && C=%i && D=%i)))"), *VariableName, Guid.A, Guid.B, Guid.C, Guid.D);
 			bUseQuotes = false;
 		}
 		else
 		{
 			FString VariableName = VarAction->GetVariableName().ToString();
-			SearchTerm = FString::Printf(TEXT("Nodes(VariableReference(MemberName=+%s))"), *VariableName);
+			SearchTerm = FString::Printf(TEXT("Nodes(VariableReference(MemberName=+\"%s\"))"), *VariableName);
 			bUseQuotes = false;
 		}
 	}
 	else if (FEdGraphSchemaAction_K2LocalVar* LocalVarAction = SelectionAsLocalVar())
 	{
-		SearchTerm = FString::Printf(TEXT("Nodes(VariableReference(MemberName=+%s && MemberScope=+%s))"), *LocalVarAction->GetVariableName().ToString(), *LocalVarAction->GetVariableScope()->GetName());
+		SearchTerm = FString::Printf(TEXT("Nodes(VariableReference(MemberName=+\"%s\" && MemberScope=+%s))"), *LocalVarAction->GetVariableName().ToString(), *LocalVarAction->GetVariableScope()->GetName());
 		bUseQuotes = false;
 	}
 	else if (FEdGraphSchemaAction_K2Delegate* DelegateAction = SelectionAsDelegate())
@@ -2506,6 +2507,12 @@ bool SMyBlueprint::IsDuplicateActionVisible() const
 
 bool SMyBlueprint::CanDuplicateAction() const
 {
+	// Cannot delete entries while not in editing mode
+	if (!IsEditingMode())
+	{
+		return false;
+	}
+
 	if (FEdGraphSchemaAction_K2Graph* GraphAction = SelectionAsGraph())
 	{
 		// Only support function graph duplication
diff --git a/Engine/Source/Editor/Kismet/Public/BlueprintEditor.h b/Engine/Source/Editor/Kismet/Public/BlueprintEditor.h
--- a/Engine/Source/Editor/Kismet/Public/BlueprintEditor.h
+++ b/Engine/Source/Editor/Kismet/Public/BlueprintEditor.h
@@ -197,6 +197,7 @@ public:
 	TSharedRef<class SKismetInspector> GetInspector() const {return Inspector.ToSharedRef();}
 	TSharedRef<class SKismetInspector> GetDefaultEditor() const {return DefaultEditor.ToSharedRef();}
 	TSharedRef<class SKismetDebuggingView> GetDebuggingView() const {return DebuggingView.ToSharedRef();}
+	TSharedRef<class SBlueprintProfilerView> GetBlueprintProfilerView() const {return BlueprintProfiler.ToSharedRef();}
 	TSharedRef<class SBlueprintPalette> GetPalette() const {return Palette.ToSharedRef();}
 	TSharedRef<class SWidget> GetCompilerResults() const {return CompilerResults.ToSharedRef();}
 	TSharedRef<class SFindInBlueprints> GetFindResults() const {return FindResults.ToSharedRef();}
@@ -507,6 +508,12 @@ public:
 	/** Can generate native code for current blueprint */
 	bool CanGenerateNativeCode() const;
 
+	/** Returns if the blueprint profiler is currently active. */
+	bool IsProfilerActive() const;
+
+	/** Toggle blueprint profiler state */
+	void ToggleProfiler();
+
 	/** 
 	 * Check to see if we can customize the SCS editor for the passed-in scene component 
 	 * @param	InComponentToCustomize	The component to check to see if a customization exists
@@ -1026,6 +1033,9 @@ protected:
 	/** Debugging window (watches, breakpoints, etc...) */
 	TSharedPtr<class SKismetDebuggingView> DebuggingView;
 
+	/** Performance Analsys View */
+	TSharedPtr<class SBlueprintProfilerView> BlueprintProfiler;
+
 	/** Palette of all classes with funcs/vars */
 	TSharedPtr<class SBlueprintPalette> Palette;
 
diff --git a/Engine/Source/Editor/Kismet/Public/BlueprintEditorTabs.h b/Engine/Source/Editor/Kismet/Public/BlueprintEditorTabs.h
--- a/Engine/Source/Editor/Kismet/Public/BlueprintEditorTabs.h
+++ b/Engine/Source/Editor/Kismet/Public/BlueprintEditorTabs.h
@@ -16,6 +16,7 @@ struct KISMET_API FBlueprintEditorTabs
 	static const FName MyBlueprintID;
 	static const FName ReplaceNodeReferencesID;
 	static const FName UserDefinedStructureID;
+	static const FName BlueprintProfilerID;
 
 	// Document tab identifiers
 	static const FName GraphEditorID;
diff --git a/Engine/Source/Editor/Kismet/Public/FindInBlueprintManager.h b/Engine/Source/Editor/Kismet/Public/FindInBlueprintManager.h
--- a/Engine/Source/Editor/Kismet/Public/FindInBlueprintManager.h
+++ b/Engine/Source/Editor/Kismet/Public/FindInBlueprintManager.h
@@ -374,7 +374,7 @@ public:
 	static FString ConvertFTextToHexString(FText InValue);
 
 	/** Returns the number of uncached Blueprints */
-	int32 GetNumberUncachedBlueprints() const { return UncachedBlueprints.Num(); };
+	int32 GetNumberUncachedBlueprints() const;
 
 	/**
 	 * Starts caching all uncached Blueprints at a rate of 1 per tick
@@ -407,7 +407,7 @@ public:
 	float GetCacheProgress() const;
 
 	/** Returns the list of Blueprint paths that failed to cache */
-	TArray<FName> GetFailedToCachePathList() const { return FailedToCachePaths; }
+	TSet<FName> GetFailedToCachePathList() const { return FailedToCachePaths; }
 
 	/** Returns the number of Blueprints that failed to cache */
 	int32 GetFailedToCacheCount() const { return FailedToCachePaths.Num(); }
@@ -419,7 +419,7 @@ public:
 	 *
 	 * @param InNumberCached		The number of Blueprints cached, to be chopped off the existing array so the rest (if any) can be finished later
 	 */
-	void FinishedCachingBlueprints(int32 InNumberCached, TArray<FName>& InFailedToCacheList);
+	void FinishedCachingBlueprints(int32 InNumberCached, TSet<FName>& InFailedToCacheList);
 
 	/** Returns TRUE if Blueprints are being cached. */
 	bool IsCacheInProgress() const;
@@ -428,7 +428,11 @@ public:
 	static FString ConvertFStringToHexString(FString InValue);
 
 	/** Given a fully constructed Find-in-Blueprint FString of searchable data, will parse and construct a JsonObject */
-	static TSharedPtr< class FJsonObject > ConvertJsonStringToObject(FString InJsonString, TMap<int32, FText>& OutFTextLookupTable);
+	static TSharedPtr< class FJsonObject > ConvertJsonStringToObject(bool bInIsVersioned, FString InJsonString, TMap<int32, FText>& OutFTextLookupTable);
+
+	void EnableGatheringData(bool bInEnableGatheringData) { bEnableGatheringData = bInEnableGatheringData; }
+
+	bool IsGatheringDataEnabled() const { return bEnableGatheringData; }
 
 private:
 	/** Initializes the FiB manager */
@@ -479,6 +483,9 @@ private:
 	void ExtractUnloadedFiBData(const FAssetData& InAssetData, const FString& InFiBData, bool bIsVersioned);
 
 protected:
+	/** Tells if gathering data is currently allowed */
+	bool bEnableGatheringData;
+
 	/** Maps the Blueprint paths to their index in the SearchArray */
 	TMap<FName, int32> SearchMap;
 
@@ -510,10 +517,10 @@ protected:
 	TWeakPtr<SFindInBlueprints> SourceCachingWidget;
 
 	/** Blueprint paths that have not been cached for searching due to lack of data, this means that they are either older Blueprints, or the DDC cannot find the data */
-	TArray<FName> UncachedBlueprints;
+	TSet<FName> UncachedBlueprints;
 
 	/** List of paths for Blueprints that failed to cache */
-	TArray<FName> FailedToCachePaths;
+	TSet<FName> FailedToCachePaths;
 
 	/** Tickable object that does the caching of uncached Blueprints at a rate of once per tick */
 	class FCacheAllBlueprintsTickableObject* CachingObject;
@@ -521,3 +528,16 @@ protected:
 	/** Mapping between a class name and its UClass instance - used for faster look up in FFindInBlueprintSearchManager::OnAssetAdded */
 	TMap<FName, const UClass*> CachedAssetClasses;
 };
+
+KISMET_API struct FDisableGatheringDataOnScope
+{
+	bool bOriginallyEnabled;
+	FDisableGatheringDataOnScope() : bOriginallyEnabled(FFindInBlueprintSearchManager::Get().IsGatheringDataEnabled())
+	{
+		FFindInBlueprintSearchManager::Get().EnableGatheringData(false);
+	}
+	~FDisableGatheringDataOnScope()
+	{
+		FFindInBlueprintSearchManager::Get().EnableGatheringData(bOriginallyEnabled);
+	}
+};
\ No newline at end of file
diff --git a/Engine/Source/Editor/Kismet/Public/ImaginaryBlueprintData.h b/Engine/Source/Editor/Kismet/Public/ImaginaryBlueprintData.h
--- a/Engine/Source/Editor/Kismet/Public/ImaginaryBlueprintData.h
+++ b/Engine/Source/Editor/Kismet/Public/ImaginaryBlueprintData.h
@@ -231,7 +231,7 @@ public:
 	bool IsHidden() const
 	{
 		// While handled separately, when hidden it should always be explicit
-		ensure(bIsHidden);
+		ensure(!bIsHidden || (bIsHidden && bIsExplicit));
 		return bIsHidden;
 	}
 
@@ -291,7 +291,7 @@ protected:
 class FImaginaryBlueprint : public FImaginaryFiBData
 {
 public:
-	FImaginaryBlueprint(FString InBlueprintName, FString InBlueprintPath, FString InBlueprintParentClass, TArray<FString>& InInterfaces, FString InUnparsedStringData);
+	FImaginaryBlueprint(FString InBlueprintName, FString InBlueprintPath, FString InBlueprintParentClass, TArray<FString>& InInterfaces, FString InUnparsedStringData, bool bInIsVersioned = true);
 
 	/** FImaginaryFiBData Interface */
 	virtual bool IsCompatibleWithFilter(ESearchQueryFilter InSearchQueryFilter) const override;
@@ -312,7 +312,7 @@ protected:
 	void ParseComponents(TSharedPtr< FJsonObject > InJsonObject, TArray< TSharedPtr<FImaginaryFiBData> >& OutParsedChildData);
 
 	/** Parses a raw string of Json to a Json object hierarchy */
-	void ParseToJson();
+	void ParseToJson(bool bInIsVersioned);
 
 protected:
 	/** The path for this Blueprint */
diff --git a/Engine/Source/Editor/Kismet/Public/SBlueprintEditorToolbar.h b/Engine/Source/Editor/Kismet/Public/SBlueprintEditorToolbar.h
--- a/Engine/Source/Editor/Kismet/Public/SBlueprintEditorToolbar.h
+++ b/Engine/Source/Editor/Kismet/Public/SBlueprintEditorToolbar.h
@@ -23,6 +23,8 @@ protected:
 
 	static void FillDebugMenu( FMenuBuilder& MenuBuilder );
 
+	static void FillProfilerMenu( FMenuBuilder& MenuBuilder );
+
 	static void FillDeveloperMenu( FMenuBuilder& MenuBuilder );
 
 private:
@@ -49,6 +51,9 @@ public:
 	TSharedPtr<FUICommandInfo> SaveOnCompile_Always;
 	TSharedPtr<FUICommandInfo> JumpToErrorNode;
 
+	/** Profiler options  */
+	TSharedPtr<FUICommandInfo> ToggleProfiler;
+
 	/** Switch between modes in the blueprint editor */
 	TSharedPtr<FUICommandInfo> SwitchToScriptingMode;
 	TSharedPtr<FUICommandInfo> SwitchToBlueprintDefaultsMode;
@@ -77,10 +82,14 @@ public:
 	void AddScriptingToolbar(TSharedPtr<FExtender> Extender);
 	void AddDebuggingToolbar(TSharedPtr<FExtender> Extender);
 	void AddComponentsToolbar(TSharedPtr<FExtender> Extender);
+	void AddProfilerToolbar(TSharedPtr<FExtender> Extender);
 
 	/** Returns the current status icon for the blueprint being edited */
 	FSlateIcon GetStatusImage() const;
 
+	/** Returns the current profiler status tool tip */
+	FText GetProfilerStatusTooltip() const;
+
 	/** Returns the current status as text for the blueprint being edited */
 	FText GetStatusTooltip() const;
 
@@ -95,6 +104,7 @@ private:
 	void FillScriptingToolbar(FToolBarBuilder& ToolbarBuilder);
 	void FillDebuggingToolbar(FToolBarBuilder& ToolbarBuilder);
 	void FillComponentsToolbar(FToolBarBuilder& ToolbarBuilder);
+	void FillProfilerToolbar(FToolBarBuilder& ToolbarBuilder);
 
 protected:
 	/** Pointer back to the blueprint editor tool that owns us */
diff --git a/Engine/Source/Editor/KismetCompiler/Private/KismetCompiler.cpp b/Engine/Source/Editor/KismetCompiler/Private/KismetCompiler.cpp
--- a/Engine/Source/Editor/KismetCompiler/Private/KismetCompiler.cpp
+++ b/Engine/Source/Editor/KismetCompiler/Private/KismetCompiler.cpp
@@ -1246,7 +1246,7 @@ void FKismetCompilerContext::PrecompileFunction(FKismetFunctionContext& Context)
 		// Inherit extra flags from the entry node
 		if (Context.EntryPoint)
 		{
-			Context.Function->FunctionFlags |= Context.EntryPoint->GetExtraFlags(Context);
+			Context.Function->FunctionFlags |= Context.EntryPoint->GetExtraFlags();
 		}
 
 		// First try to get the overriden function from the super class
@@ -1354,7 +1354,7 @@ void FKismetCompilerContext::PrecompileFunction(FKismetFunctionContext& Context)
 		CreateLocalVariablesForFunction(Context);
 
 		//Validate AccessSpecifier
-		const uint32 AccessSpecifierFlag = FUNC_AccessSpecifiers & Context.EntryPoint->GetExtraFlags(Context);
+		const uint32 AccessSpecifierFlag = FUNC_AccessSpecifiers & Context.EntryPoint->GetExtraFlags();
 		const bool bAcceptedAccessSpecifier = 
 			(0 == AccessSpecifierFlag) || (FUNC_Public == AccessSpecifierFlag) || (FUNC_Protected == AccessSpecifierFlag) || (FUNC_Private == AccessSpecifierFlag);
 		if(!bAcceptedAccessSpecifier)
@@ -1473,7 +1473,7 @@ void FKismetCompilerContext::CompileFunction(FKismetFunctionContext& Context)
 		}
 
 		// Debug opcode insertion point
-		if (Context.bCreateDebugData && !IsNodePure(Node))
+		if (Context.IsDebuggingOrInstrumentationRequired() && !IsNodePure(Node))
 		{
 			bool bEmitDebuggingSite = true;
 
@@ -1488,7 +1488,7 @@ void FKismetCompilerContext::CompileFunction(FKismetFunctionContext& Context)
 			if (bEmitDebuggingSite)
 			{
 				FBlueprintCompiledStatement& Statement = Context.AppendStatementForNode(Node);
-				Statement.Type = KCST_DebugSite;
+				Statement.Type = Context.GetBreakpointType();
 				Statement.Comment = NodeComment;
 			}
 		}
@@ -1943,107 +1943,111 @@ void FKismetCompilerContext::ExpandTimelineNodes(UEdGraph* SourceGraph)
 	{
 		UK2Node_Timeline* const TimelineNode = TimelinePair.Node;
 		UTimelineTemplate* Timeline = TimelinePair.Template;
-		const FString TimelineNameString = TimelineNode->TimelineName.ToString();
 
-		UEdGraphPin* PlayPin = TimelineNode->GetPlayPin();
-		bool bPlayPinConnected = (PlayPin->LinkedTo.Num() > 0);
+		if (bIsFullCompile)
+		{
+			const FString TimelineNameString = TimelineNode->TimelineName.ToString();
 
-		UEdGraphPin* PlayFromStartPin = TimelineNode->GetPlayFromStartPin();
-		bool bPlayFromStartPinConnected = (PlayFromStartPin->LinkedTo.Num() > 0);
+			UEdGraphPin* PlayPin = TimelineNode->GetPlayPin();
+			bool bPlayPinConnected = (PlayPin->LinkedTo.Num() > 0);
 
-		UEdGraphPin* StopPin = TimelineNode->GetStopPin();
-		bool bStopPinConnected = (StopPin->LinkedTo.Num() > 0);
+			UEdGraphPin* PlayFromStartPin = TimelineNode->GetPlayFromStartPin();
+			bool bPlayFromStartPinConnected = (PlayFromStartPin->LinkedTo.Num() > 0);
 
-		UEdGraphPin* ReversePin = TimelineNode->GetReversePin();
-		bool bReversePinConnected = (ReversePin->LinkedTo.Num() > 0);
+			UEdGraphPin* StopPin = TimelineNode->GetStopPin();
+			bool bStopPinConnected = (StopPin->LinkedTo.Num() > 0);
 
-		UEdGraphPin* ReverseFromEndPin = TimelineNode->GetReverseFromEndPin();
-		bool bReverseFromEndPinConnected = (ReverseFromEndPin->LinkedTo.Num() > 0);
+			UEdGraphPin* ReversePin = TimelineNode->GetReversePin();
+			bool bReversePinConnected = (ReversePin->LinkedTo.Num() > 0);
 
-		UEdGraphPin* SetTimePin = TimelineNode->GetSetNewTimePin();
-		bool bSetNewTimePinConnected = (SetTimePin->LinkedTo.Num() > 0);
+			UEdGraphPin* ReverseFromEndPin = TimelineNode->GetReverseFromEndPin();
+			bool bReverseFromEndPinConnected = (ReverseFromEndPin->LinkedTo.Num() > 0);
 
-		UEdGraphPin* UpdatePin = TimelineNode->GetUpdatePin();
-		bool bUpdatePinConnected = (UpdatePin->LinkedTo.Num() > 0);
+			UEdGraphPin* SetTimePin = TimelineNode->GetSetNewTimePin();
+			bool bSetNewTimePinConnected = (SetTimePin->LinkedTo.Num() > 0);
 
-		UEdGraphPin* FinishedPin = TimelineNode->GetFinishedPin();
-		bool bFinishedPinConnected = (FinishedPin->LinkedTo.Num() > 0);
-		
+			UEdGraphPin* UpdatePin = TimelineNode->GetUpdatePin();
+			bool bUpdatePinConnected = (UpdatePin->LinkedTo.Num() > 0);
 
-		// Set the timeline template as wired/not wired for component pruning later
-		const bool bWiredIn = bPlayPinConnected || bPlayFromStartPinConnected || bStopPinConnected || bReversePinConnected || bReverseFromEndPinConnected || bSetNewTimePinConnected;
-		const bool bWiredOut = bUpdatePinConnected || bFinishedPinConnected;
-		const bool bPlayWired = Timeline->bAutoPlay;
-		const bool bReferenced = TimelinePlayNodes.Find( TimelineNode->TimelineName ) != INDEX_NONE;
+			UEdGraphPin* FinishedPin = TimelineNode->GetFinishedPin();
+			bool bFinishedPinConnected = (FinishedPin->LinkedTo.Num() > 0);
 
-		Timeline->bValidatedAsWired = bWiredIn || bReferenced || ( bPlayWired && bWiredOut );
 
-		// Only create nodes for play/stop if they are actually connected - otherwise we get a 'unused node being pruned' warning
-		if(bWiredIn)
-		{
-			// First create 'get var' node to get the timeline object
-			UK2Node_VariableGet* GetTimelineNode = SpawnIntermediateNode<UK2Node_VariableGet>(TimelineNode, SourceGraph);
-			GetTimelineNode->VariableReference.SetSelfMember(TimelineNode->TimelineName);
-			GetTimelineNode->AllocateDefaultPins();
-
-			// Debug data: Associate the timeline node instance with the property that was created earlier
-			UProperty* AssociatedTimelineInstanceProperty = TimelineToMemberVariableMap.FindChecked(Timeline);
-			if (AssociatedTimelineInstanceProperty != NULL)
-			{
-				UObject* TrueSourceObject = MessageLog.FindSourceObject(TimelineNode);
-				NewClass->GetDebugData().RegisterClassPropertyAssociation(TrueSourceObject, AssociatedTimelineInstanceProperty);
-			}
+			// Set the timeline template as wired/not wired for component pruning later
+			const bool bWiredIn = bPlayPinConnected || bPlayFromStartPinConnected || bStopPinConnected || bReversePinConnected || bReverseFromEndPinConnected || bSetNewTimePinConnected;
+			const bool bWiredOut = bUpdatePinConnected || bFinishedPinConnected;
+			const bool bPlayWired = Timeline->bAutoPlay;
+			const bool bReferenced = TimelinePlayNodes.Find(TimelineNode->TimelineName) != INDEX_NONE;
 
-			// Get the variable output pin
-			UEdGraphPin* TimelineVarPin = GetTimelineNode->FindPin(TimelineNameString);
+			Timeline->bValidatedAsWired = bWiredIn || bReferenced || (bPlayWired && bWiredOut);
 
-			// This might fail if this is the first compile after adding the timeline (property doesn't exist yet) - in that case, manually add the output pin
-			if(TimelineVarPin == NULL)
+			// Only create nodes for play/stop if they are actually connected - otherwise we get a 'unused node being pruned' warning
+			if (bWiredIn)
 			{
-				TimelineVarPin = GetTimelineNode->CreatePin(EGPD_Output, Schema->PC_Object, TEXT(""), UTimelineComponent::StaticClass(), false, false, TimelineNode->TimelineName.ToString());
-			}
+				// First create 'get var' node to get the timeline object
+				UK2Node_VariableGet* GetTimelineNode = SpawnIntermediateNode<UK2Node_VariableGet>(TimelineNode, SourceGraph);
+				GetTimelineNode->VariableReference.SetSelfMember(TimelineNode->TimelineName);
+				GetTimelineNode->AllocateDefaultPins();
 
-			if(bPlayPinConnected)
-			{
-				static FName PlayName(GET_FUNCTION_NAME_CHECKED(UTimelineComponent, Play));
-				CreateCallTimelineFunction(TimelineNode, SourceGraph, PlayName, TimelineVarPin, PlayPin);
-			}
+				// Debug data: Associate the timeline node instance with the property that was created earlier
+				UProperty* AssociatedTimelineInstanceProperty = TimelineToMemberVariableMap.FindChecked(Timeline);
+				if (AssociatedTimelineInstanceProperty != NULL)
+				{
+					UObject* TrueSourceObject = MessageLog.FindSourceObject(TimelineNode);
+					NewClass->GetDebugData().RegisterClassPropertyAssociation(TrueSourceObject, AssociatedTimelineInstanceProperty);
+				}
 
-			if(bPlayFromStartPinConnected)
-			{
-				static FName PlayFromStartName(GET_FUNCTION_NAME_CHECKED(UTimelineComponent, PlayFromStart));
-				CreateCallTimelineFunction(TimelineNode, SourceGraph, PlayFromStartName, TimelineVarPin, PlayFromStartPin);
-			}
+				// Get the variable output pin
+				UEdGraphPin* TimelineVarPin = GetTimelineNode->FindPin(TimelineNameString);
 
-			if(bStopPinConnected)
-			{
-				static FName StopName(GET_FUNCTION_NAME_CHECKED(UTimelineComponent, Stop));
-				CreateCallTimelineFunction(TimelineNode, SourceGraph, StopName, TimelineVarPin, StopPin);
-			}
+				// This might fail if this is the first compile after adding the timeline (property doesn't exist yet) - in that case, manually add the output pin
+				if (TimelineVarPin == NULL)
+				{
+					TimelineVarPin = GetTimelineNode->CreatePin(EGPD_Output, Schema->PC_Object, TEXT(""), UTimelineComponent::StaticClass(), false, false, TimelineNode->TimelineName.ToString());
+				}
 
-			if (bReversePinConnected)
-			{
-				static FName ReverseName(GET_FUNCTION_NAME_CHECKED(UTimelineComponent, Reverse));
-				CreateCallTimelineFunction(TimelineNode, SourceGraph, ReverseName, TimelineVarPin, ReversePin);
-			}
+				if (bPlayPinConnected)
+				{
+					static FName PlayName(GET_FUNCTION_NAME_CHECKED(UTimelineComponent, Play));
+					CreateCallTimelineFunction(TimelineNode, SourceGraph, PlayName, TimelineVarPin, PlayPin);
+				}
 
-			if (bReverseFromEndPinConnected)
-			{
-				static FName ReverseFromEndName(GET_FUNCTION_NAME_CHECKED(UTimelineComponent, ReverseFromEnd));
-				CreateCallTimelineFunction(TimelineNode, SourceGraph, ReverseFromEndName, TimelineVarPin, ReverseFromEndPin);
-			}
+				if (bPlayFromStartPinConnected)
+				{
+					static FName PlayFromStartName(GET_FUNCTION_NAME_CHECKED(UTimelineComponent, PlayFromStart));
+					CreateCallTimelineFunction(TimelineNode, SourceGraph, PlayFromStartName, TimelineVarPin, PlayFromStartPin);
+				}
 
-			if (bSetNewTimePinConnected)
-			{
-				UEdGraphPin* NewTimePin = TimelineNode->GetNewTimePin();
+				if (bStopPinConnected)
+				{
+					static FName StopName(GET_FUNCTION_NAME_CHECKED(UTimelineComponent, Stop));
+					CreateCallTimelineFunction(TimelineNode, SourceGraph, StopName, TimelineVarPin, StopPin);
+				}
 
-				static FName SetNewTimeName(GET_FUNCTION_NAME_CHECKED(UTimelineComponent, SetNewTime));
-				UK2Node_CallFunction* CallNode = CreateCallTimelineFunction(TimelineNode, SourceGraph, SetNewTimeName, TimelineVarPin, SetTimePin);
+				if (bReversePinConnected)
+				{
+					static FName ReverseName(GET_FUNCTION_NAME_CHECKED(UTimelineComponent, Reverse));
+					CreateCallTimelineFunction(TimelineNode, SourceGraph, ReverseName, TimelineVarPin, ReversePin);
+				}
 
-				if (CallNode && NewTimePin)
+				if (bReverseFromEndPinConnected)
 				{
-					UEdGraphPin* InputPin = CallNode->FindPinChecked(TEXT("NewTime"));
-					MovePinLinksToIntermediate(*NewTimePin, *InputPin);
+					static FName ReverseFromEndName(GET_FUNCTION_NAME_CHECKED(UTimelineComponent, ReverseFromEnd));
+					CreateCallTimelineFunction(TimelineNode, SourceGraph, ReverseFromEndName, TimelineVarPin, ReverseFromEndPin);
+				}
+
+				if (bSetNewTimePinConnected)
+				{
+					UEdGraphPin* NewTimePin = TimelineNode->GetNewTimePin();
+
+					static FName SetNewTimeName(GET_FUNCTION_NAME_CHECKED(UTimelineComponent, SetNewTime));
+					UK2Node_CallFunction* CallNode = CreateCallTimelineFunction(TimelineNode, SourceGraph, SetNewTimeName, TimelineVarPin, SetTimePin);
+
+					if (CallNode && NewTimePin)
+					{
+						UEdGraphPin* InputPin = CallNode->FindPinChecked(TEXT("NewTime"));
+						MovePinLinksToIntermediate(*NewTimePin, *InputPin);
+					}
 				}
 			}
 		}
@@ -2150,7 +2154,7 @@ void FKismetCompilerContext::CreateFunctionStubForEvent(UK2Node_Event* SrcEventN
 	ChildStubGraph->SetFlags(RF_Transient);
 	MessageLog.NotifyIntermediateObjectCreation(ChildStubGraph, SrcEventNode);
 
-	FKismetFunctionContext& StubContext = *new (FunctionList) FKismetFunctionContext(MessageLog, Schema, NewClass, Blueprint, CompileOptions.DoesRequireCppCodeGeneration());
+	FKismetFunctionContext& StubContext = *new (FunctionList) FKismetFunctionContext(MessageLog, Schema, NewClass, Blueprint, CompileOptions.DoesRequireCppCodeGeneration(), false);
 	StubContext.SourceGraph = ChildStubGraph;
 
 	// A stub graph has no visual representation and is thus not suited to be debugged via the debugger
@@ -2185,9 +2189,9 @@ void FKismetCompilerContext::CreateFunctionStubForEvent(UK2Node_Event* SrcEventN
 	EntryNode->SignatureName = SrcEventNode->EventReference.GetMemberName();
 	EntryNode->CustomGeneratedFunctionName = EventNodeName;
 
-	if(!SrcEventNode->bOverrideFunction && SrcEventNode->IsUsedByAuthorityOnlyDelegate())
+	if (!SrcEventNode->bOverrideFunction && SrcEventNode->IsUsedByAuthorityOnlyDelegate())
 	{
-		EntryNode->SetExtraFlags(EntryNode->GetFunctionFlags() | FUNC_BlueprintAuthorityOnly);
+		EntryNode->AddExtraFlags(FUNC_BlueprintAuthorityOnly);
 	}
 
 	// If this is a customizable event, make sure to copy over the user defined pins
@@ -2340,12 +2344,12 @@ void FKismetCompilerContext::ExpansionStep(UEdGraph* Graph, bool bAllowUbergraph
 				Node->ExpandNode(*this, Graph);
 			}
 		}
+	}
 
-		if (bAllowUbergraphExpansions)
-		{
-			// Expand timeline nodes
-			ExpandTimelineNodes(Graph);
-		}
+	if (bAllowUbergraphExpansions)
+	{
+		// Expand timeline nodes, in skeleton classes only the events will be generated
+		ExpandTimelineNodes(Graph);
 	}
 }
 
@@ -2510,7 +2514,7 @@ void FKismetCompilerContext::CreateAndProcessUbergraph()
 
 		// Do some cursory validation (pin types match, inputs to outputs, pins never point to their parent node, etc...)
 		{
-			UbergraphContext = new (FunctionList)FKismetFunctionContext(MessageLog, Schema, NewClass, Blueprint, CompileOptions.DoesRequireCppCodeGeneration());
+			UbergraphContext = new (FunctionList)FKismetFunctionContext(MessageLog, Schema, NewClass, Blueprint, CompileOptions.DoesRequireCppCodeGeneration(), CompileOptions.IsInstrumentationActive());
 			UbergraphContext->SourceGraph = ConsolidatedEventGraph;
 			UbergraphContext->MarkAsEventGraph();
 			UbergraphContext->MarkAsInternalOrCppUseOnly();
@@ -2625,9 +2629,22 @@ void FKismetCompilerContext::ExpandTunnelsAndMacros(UEdGraph* SourceGraph)
 				continue;
 			}
 
+			UBlueprint* MacroBlueprint = FBlueprintEditorUtils::FindBlueprintForGraph(MacroGraph);
+			// unfortunately, you may be expanding a macro that has yet to be
+			// regenerated on load (thanks cyclic dependencies!), and in certain
+			// cases the nodes found within the macro may be out of date 
+			// (function signatures, etc.), so let's force a reconstruct of the 
+			// nodes we inject from the macro (just in case)
+			const bool bForceRegenNodes = bIsLoading && MacroBlueprint && (MacroBlueprint != Blueprint) && !MacroBlueprint->bHasBeenRegenerated;
+
 			// Clone the macro graph, then move all of its children, keeping a list of nodes from the macro
 			UEdGraph* ClonedGraph = FEdGraphUtilities::CloneGraph(MacroGraph, NULL, &MessageLog, true);
 
+			for (int32 I = 0; I < ClonedGraph->Nodes.Num(); ++I)
+			{
+				MacroGeneratedNodes.Add(ClonedGraph->Nodes[I], CurrentNode);
+			}
+
 			TArray<UEdGraphNode*> MacroNodes(ClonedGraph->Nodes);
 
 			// resolve any wildcard pins in the nodes cloned from the macro
@@ -2732,6 +2749,11 @@ void FKismetCompilerContext::ExpandTunnelsAndMacros(UEdGraph* SourceGraph)
 				
 				if( DuplicatedNode != NULL )
 				{
+					if (bForceRegenNodes)
+					{
+						DuplicatedNode->ReconstructNode();
+					}
+
 					// Record the source node mapping for the intermediate node first, as it's going to be overwritten through the MessageLog below
 					UEdGraphNode* MacroSourceNode = Cast<UEdGraphNode>(MessageLog.FindSourceObject(DuplicatedNode));
 					if (MacroSourceNode)
@@ -2856,7 +2878,7 @@ void FKismetCompilerContext::ProcessOneFunctionGraph(UEdGraph* SourceGraph, bool
 	//       been added to the class yet, etc.)
 	if ((CompileOptions.CompileType == EKismetCompileType::SkeletonOnly) || ValidateGraphIsWellFormed(FunctionGraph))
 	{
-		FKismetFunctionContext& Context = *new (FunctionList)FKismetFunctionContext(MessageLog, Schema, NewClass, Blueprint, CompileOptions.DoesRequireCppCodeGeneration());
+		FKismetFunctionContext& Context = *new (FunctionList)FKismetFunctionContext(MessageLog, Schema, NewClass, Blueprint, CompileOptions.DoesRequireCppCodeGeneration(), CompileOptions.IsInstrumentationActive());
 		Context.SourceGraph = FunctionGraph;
 
 		if(FBlueprintEditorUtils::IsDelegateSignatureGraph(SourceGraph))
@@ -2958,7 +2980,7 @@ void FKismetCompilerContext::CreateFunctionList()
 
 FKismetFunctionContext* FKismetCompilerContext::CreateFunctionContext()
 {
-	return new (FunctionList)FKismetFunctionContext(MessageLog, Schema, NewClass, Blueprint, CompileOptions.DoesRequireCppCodeGeneration());
+	return new (FunctionList)FKismetFunctionContext(MessageLog, Schema, NewClass, Blueprint, CompileOptions.DoesRequireCppCodeGeneration(), CompileOptions.IsInstrumentationActive());
 }
 
 /** Compile a blueprint into a class and a set of functions */
@@ -3072,6 +3094,9 @@ void FKismetCompilerContext::Compile()
 		}
 	}
 
+	// Declare class wants profiling data.
+	TargetClass->bHasInstrumentation = CompileOptions.IsInstrumentationActive();
+
 	// Early validation
 	if (CompileOptions.CompileType == EKismetCompileType::Full)
 	{
@@ -3800,5 +3825,22 @@ bool FKismetCompilerContext::UsePersistentUberGraphFrame() const
 	return UBlueprintGeneratedClass::UsePersistentUberGraphFrame() && !CompileOptions.DoesRequireCppCodeGeneration();
 }
 
+FString FKismetCompilerContext::GetGuid(const UEdGraphNode* Node) const
+{
+	// We need a unique, deterministic name for the properties we're generating, but the chance of collision is small
+	// so I think we can get away with stomping part of a guid with a hash:
+	uint32 ResultCRC = FCrc::MemCrc32(&(Node->NodeGuid), sizeof(Node->NodeGuid));
+	UEdGraphNode* const* SourceNode = MacroGeneratedNodes.Find(Node);
+	while (SourceNode && *SourceNode)
+	{
+		ResultCRC = FCrc::MemCrc32(&((*SourceNode)->NodeGuid), sizeof((*SourceNode)->NodeGuid), ResultCRC);
+		SourceNode = MacroGeneratedNodes.Find(*SourceNode);
+	}
+
+	FGuid Ret = Node->NodeGuid;
+	Ret.D = ResultCRC;
+	return Ret.ToString();
+}
+
 #undef LOCTEXT_NAMESPACE
 //////////////////////////////////////////////////////////////////////////
diff --git a/Engine/Source/Editor/KismetCompiler/Private/KismetCompilerMisc.cpp b/Engine/Source/Editor/KismetCompiler/Private/KismetCompilerMisc.cpp
--- a/Engine/Source/Editor/KismetCompiler/Private/KismetCompilerMisc.cpp
+++ b/Engine/Source/Editor/KismetCompiler/Private/KismetCompilerMisc.cpp
@@ -688,6 +688,8 @@ void FKismetCompilerUtilities::CreateObjectAssignmentStatement(FKismetFunctionCo
 		ClassTerm->bIsLiteral = true;
 		ClassTerm->Source = DstTerm->Source;
 		ClassTerm->ObjectLiteral = OutputObjClass;
+		check(Context.Schema);
+		ClassTerm->Type.PinCategory = Context.Schema->PC_Class;
 
 		EKismetCompiledStatementType CastOpType = bIsOutputInterface ? KCST_CastObjToInterface : KCST_CastInterfaceToObj;
 		FBlueprintCompiledStatement& CastStatement = Context.AppendStatementForNode(Node);
@@ -855,7 +857,11 @@ UProperty* FKismetCompilerUtilities::CreatePropertyOnScope(UStruct* Scope, const
 			if (Type.PinCategory == Schema->PC_AssetClass)
 			{
 				auto AssetClassProperty = NewObject<UAssetClassProperty>(PropertyScope, ValidatedPropertyName, ObjectFlags);
-				AssetClassProperty->MetaClass = SubType;
+				// we want to use this setter function instead of setting the 
+				// MetaClass member directly, because it properly handles  
+				// placeholder classes (classes that are stubbed in during load)
+				AssetClassProperty->SetMetaClass(SubType);
+				AssetClassProperty->PropertyClass = UClass::StaticClass();
 				NewProperty = AssetClassProperty;
 			}
 			else
@@ -1236,10 +1242,10 @@ FBlueprintCompiledStatement& FNodeHandlingFunctor::GenerateSimpleThenGoto(FKisme
 		TargetNode = ThenExecPin->LinkedTo[0]->GetOwningNode();
 	}
 
-	if (Context.bCreateDebugData)
+	if (Context.bCreateDebugData || Context.bInstrumentScriptCode)
 	{
 		FBlueprintCompiledStatement& TraceStatement = Context.AppendStatementForNode(&Node);
-		TraceStatement.Type = KCST_WireTraceSite;
+		TraceStatement.Type = Context.GetWireTraceType();
 		TraceStatement.Comment = Node.NodeComment.IsEmpty() ? Node.GetName() : Node.NodeComment;
 	}
 
@@ -1372,7 +1378,7 @@ KISMETCOMPILER_API FString FNetNameMapping::MakeBaseName<UEdGraphNode>(const UEd
 //////////////////////////////////////////////////////////////////////////
 // FKismetFunctionContext
 
-FKismetFunctionContext::FKismetFunctionContext(FCompilerResultsLog& InMessageLog, UEdGraphSchema_K2* InSchema, UBlueprintGeneratedClass* InNewClass, UBlueprint* InBlueprint, bool bInGeneratingCpp)
+FKismetFunctionContext::FKismetFunctionContext(FCompilerResultsLog& InMessageLog, UEdGraphSchema_K2* InSchema, UBlueprintGeneratedClass* InNewClass, UBlueprint* InBlueprint, bool bInGeneratingCpp, bool bInWantsInstrumentation)
 	: Blueprint(InBlueprint)
 	, SourceGraph(NULL)
 	, EntryPoint(NULL)
@@ -1385,6 +1391,7 @@ FKismetFunctionContext::FKismetFunctionContext(FCompilerResultsLog& InMessageLog
 	, bIsInterfaceStub(false)
 	, bIsConstFunction(false)
 	, bEnforceConstCorrectness(false)
+	, bInstrumentScriptCode(bInWantsInstrumentation)
 	// only need debug-data when running in the editor app:
 	, bCreateDebugData(GIsEditor && !IsRunningCommandlet())
 	, bIsSimpleStubGraphWithNoParams(false)
diff --git a/Engine/Source/Editor/KismetCompiler/Private/KismetCompilerVMBackend.cpp b/Engine/Source/Editor/KismetCompiler/Private/KismetCompilerVMBackend.cpp
--- a/Engine/Source/Editor/KismetCompiler/Private/KismetCompilerVMBackend.cpp
+++ b/Engine/Source/Editor/KismetCompiler/Private/KismetCompilerVMBackend.cpp
@@ -939,6 +939,7 @@ public:
 		if (bIsArray)
 		{
 			Writer << EX_Let;
+			ensure(DestinationExpression->AssociatedVarProperty);
 			Writer << DestinationExpression->AssociatedVarProperty;
 		}
 		else if (bIsMulticastDelegate)
@@ -967,6 +968,7 @@ public:
 		else
 		{
 			Writer << EX_Let;
+			ensure(DestinationExpression->AssociatedVarProperty);
 			Writer << DestinationExpression->AssociatedVarProperty;
 
 		}
@@ -1313,6 +1315,18 @@ public:
 		Writer.CommitSkip(PatchUpNeededAtOffset, Writer.ScriptBuffer.Num());
 	}
 
+	void EmitInstrumentation(FBlueprintCompiledStatement& Statement)
+	{
+		int32 EventType = 0;
+		switch (Statement.Type)
+		{
+			case KCST_InstrumentedWireExit:		EventType = EScriptInstrumentation::NodeExit; break;
+			case KCST_InstrumentedWireEntry:	EventType = EScriptInstrumentation::NodeEntry; break;
+		}
+		Writer << EX_InstrumentationEvent;
+		Writer << EventType;
+	}
+
 	void PushReturnAddress(FBlueprintCompiledStatement& ReturnTarget)
 	{
 		Writer << EX_PushExecutionFlow;
@@ -1467,6 +1481,42 @@ public:
 		case KCST_SwitchValue:
 			EmitSwitchValue(Statement);
 			break;
+		case KCST_InstrumentedWireExit:
+			{
+				UEdGraphPin const* TrueSourcePin = Cast<UEdGraphPin const>(FunctionContext.MessageLog.FindSourceObject(Statement.ExecContext));
+				if (TrueSourcePin)
+				{
+					int32 Offset = Writer.ScriptBuffer.Num() + sizeof(int32);
+					ClassBeingBuilt->GetDebugData().RegisterPinToCodeAssociation(TrueSourcePin, FunctionContext.Function, Offset);
+				}
+			}
+			// no break, continue down.
+		case KCST_InstrumentedWireEntry:
+			{
+				if (SourceNode != NULL)
+				{
+					// Record where this NOP is
+					UEdGraphNode* TrueSourceNode = Cast<UEdGraphNode>(FunctionContext.MessageLog.FindSourceObject(SourceNode));
+					if (TrueSourceNode)
+					{
+						// If this is a debug site for an expanded macro instruction, there should also be a macro source node associated with it
+						UEdGraphNode* MacroSourceNode = Cast<UEdGraphNode>(CompilerContext.MessageLog.FinalNodeBackToMacroSourceMap.FindSourceObject(SourceNode));
+						if (MacroSourceNode == SourceNode)
+						{
+							// The function above will return the given node if not found in the map. In that case there is no associated source macro node, so we clear it.
+							MacroSourceNode = NULL;
+						}
+	
+						TArray<TWeakObjectPtr<UEdGraphNode>> MacroInstanceNodes;
+						int32 Offset = Writer.ScriptBuffer.Num() + sizeof(int32);
+						ClassBeingBuilt->GetDebugData().RegisterNodeToCodeAssociation(TrueSourceNode, MacroSourceNode, MacroInstanceNodes, FunctionContext.Function, Offset, false);
+					}
+				}
+				// Emit Statement
+				EmitInstrumentation(Statement);
+				break;
+			}
+
 		default:
 			UE_LOG(LogK2Compiler, Warning, TEXT("VM backend encountered unsupported statement type %d"), (int32)Statement.Type);
 		}
diff --git a/Engine/Source/Editor/KismetCompiler/Public/BlueprintCompiledStatement.h b/Engine/Source/Editor/KismetCompiler/Public/BlueprintCompiledStatement.h
--- a/Engine/Source/Editor/KismetCompiler/Public/BlueprintCompiledStatement.h
+++ b/Engine/Source/Editor/KismetCompiler/Public/BlueprintCompiledStatement.h
@@ -37,6 +37,9 @@ enum EKismetCompiledStatementType
 	KCST_GotoReturn = 27,	// goto ReturnLabel
 	KCST_GotoReturnIfNot = 28, // [if (!TargetObject->TargetProperty)] goto TargetLabel
 	KCST_SwitchValue = 29,
+	// Kismet instrumentation extensions
+	KCST_InstrumentedWireEntry,	// Instrumented wiretrace entry
+	KCST_InstrumentedWireExit	// Instrumented wiretrace exit
 };
 
 //@TODO: Too rigid / icky design
diff --git a/Engine/Source/Editor/KismetCompiler/Public/KismetCompiledFunctionContext.h b/Engine/Source/Editor/KismetCompiler/Public/KismetCompiledFunctionContext.h
--- a/Engine/Source/Editor/KismetCompiler/Public/KismetCompiledFunctionContext.h
+++ b/Engine/Source/Editor/KismetCompiler/Public/KismetCompiledFunctionContext.h
@@ -78,6 +78,7 @@ public:
 	bool bIsInterfaceStub;
 	bool bIsConstFunction;
 	bool bEnforceConstCorrectness;
+	bool bInstrumentScriptCode;
 	bool bCreateDebugData;
 	bool bIsSimpleStubGraphWithNoParams;
 	uint32 NetFlags;
@@ -99,7 +100,7 @@ public:
 	//Does this function use requires FlowStack ?
 	bool bUseFlowStack;
 public:
-	FKismetFunctionContext(FCompilerResultsLog& InMessageLog, UEdGraphSchema_K2* InSchema, UBlueprintGeneratedClass* InNewClass, UBlueprint* InBlueprint, bool bInGeneratingCpp);
+	FKismetFunctionContext(FCompilerResultsLog& InMessageLog, UEdGraphSchema_K2* InSchema, UBlueprintGeneratedClass* InNewClass, UBlueprint* InBlueprint, bool bInGeneratingCpp, bool bInWantsInstrumentation);
 	
 	~FKismetFunctionContext();
 
@@ -172,6 +173,21 @@ public:
 		NetFlags = InFunctionFlags & (FUNC_NetFuncFlags);
 	}
 
+	bool IsDebuggingOrInstrumentationRequired() const
+	{
+		return bInstrumentScriptCode || bCreateDebugData;
+	}
+
+	EKismetCompiledStatementType GetWireTraceType() const
+	{
+		return bInstrumentScriptCode ? KCST_InstrumentedWireExit : KCST_WireTraceSite;
+	}
+
+	EKismetCompiledStatementType GetBreakpointType() const
+	{
+		return bInstrumentScriptCode ? KCST_InstrumentedWireEntry : KCST_DebugSite;
+	}
+
 	uint32 GetNetFlags() const
 	{
 		return NetFlags;
@@ -305,7 +321,7 @@ public:
 	void InsertWireTrace(FBlueprintCompiledStatement* GotoStatement, UEdGraphPin* AssociatedExecPin)
 	{
 		// only need wire traces if we're debugging the blueprint is available (not for cooked builds)
-		if (bCreateDebugData && (AssociatedExecPin != NULL))
+		if (IsDebuggingOrInstrumentationRequired() && (AssociatedExecPin != NULL))
 		{
 			UEdGraphNode* PreJumpNode = AssociatedExecPin->GetOwningNode();
 
@@ -323,14 +339,14 @@ public:
 					}
 
 					// if a wire-trace has already been inserted for us
-					if ((PrevStatement != NULL) && (PrevStatement->Type == KCST_WireTraceSite))
+					if ((PrevStatement != NULL) && PrevStatement->Type == GetWireTraceType())
 					{
 						PrevStatement->ExecContext = AssociatedExecPin;
 					}
 					else if (PrevStatement != NULL)
 					{
 						FBlueprintCompiledStatement* TraceStatement = new FBlueprintCompiledStatement();
-						TraceStatement->Type        = KCST_WireTraceSite;
+						TraceStatement->Type        = GetWireTraceType();
 						TraceStatement->Comment     = PreJumpNode->NodeComment.IsEmpty() ? PreJumpNode->GetName() : PreJumpNode->NodeComment;
 						TraceStatement->ExecContext = AssociatedExecPin;
 
@@ -342,7 +358,7 @@ public:
 			}
 		}
 	}
-	
+
 	/** Looks for a pin of the given name, erroring if the pin is not found or if the direction doesn't match (doesn't verify the pin type) */
 	UEdGraphPin* FindRequiredPinByName(const UEdGraphNode* Node, const FString& PinName, EEdGraphPinDirection RequiredDirection = EGPD_MAX)
 	{
diff --git a/Engine/Source/Editor/KismetCompiler/Public/KismetCompiler.h b/Engine/Source/Editor/KismetCompiler/Public/KismetCompiler.h
--- a/Engine/Source/Editor/KismetCompiler/Public/KismetCompiler.h
+++ b/Engine/Source/Editor/KismetCompiler/Public/KismetCompiler.h
@@ -81,6 +81,13 @@ public:
 	TMap<UEdGraphNode*, UEdGraphNode*> CallsIntoUbergraph;
 	int32 bIsFullCompile:1;
 
+	// Map that can be used to find the macro node that spawned a provided node, 
+	// if any. Macro instances can have more macros inside of them, so entries 
+	// in this map may chain (i.e. values may also need to be used as keys to find
+	// the full chain). Used to generate deterministic, unique identifiers for 
+	// properties generated by nodes.
+	TMap<UEdGraphNode*, UEdGraphNode*> MacroGeneratedNodes;
+
 	// Map from a name to the number of times it's been 'created' (identical nodes create the same variable names, so they need something appended)
 	FNetNameMapping ClassScopeNetNameMap;
 public:
@@ -136,6 +143,7 @@ public:
 
 	bool UsePersistentUberGraphFrame() const;
 
+	FString GetGuid(const UEdGraphNode* Node) const;
 protected:
 	virtual UEdGraphSchema_K2* CreateSchema();
 	virtual void PostCreateSchema();
diff --git a/Engine/Source/Editor/MainFrame/Private/Frame/MainFrameActions.cpp b/Engine/Source/Editor/MainFrame/Private/Frame/MainFrameActions.cpp
--- a/Engine/Source/Editor/MainFrame/Private/Frame/MainFrameActions.cpp
+++ b/Engine/Source/Editor/MainFrame/Private/Frame/MainFrameActions.cpp
@@ -553,6 +553,11 @@ void FMainFrameActionCallbacks::PackageProject( const FName InPlatformInfoName )
 		OptionalParams += TEXT(" -nodebuginfo");
 	}
 
+	if (PackagingSettings->bNativizeBlueprintAssets)
+	{
+		OptionalParams += TEXT(" -nativizeAssets");
+	}
+
 
 	bool bTargetPlatformCanUseCrashReporter = true;
 	if (PlatformInfo->TargetPlatformName == FName("WindowsNoEditor") && PlatformInfo->PlatformFlavor == TEXT("Win32"))
@@ -579,7 +584,7 @@ void FMainFrameActionCallbacks::PackageProject( const FName InPlatformInfoName )
 	}
 
 	// only build if the project has code that might need to be built
-	if (bProjectHasCode || (!FApp::GetEngineIsPromotedBuild() && !FApp::IsEngineInstalled()))
+	if (bProjectHasCode || (!FApp::GetEngineIsPromotedBuild() && !FApp::IsEngineInstalled()) || PackagingSettings->bNativizeBlueprintAssets)
 	{
 		OptionalParams += TEXT(" -build");
 	}
@@ -598,7 +603,7 @@ void FMainFrameActionCallbacks::PackageProject( const FName InPlatformInfoName )
 	int32 NumCookers = GetDefault<UEditorExperimentalSettings>()->MultiProcessCooking;
 	if (NumCookers > 0 )
 	{
-		OptionalParams += FString::Printf(TEXT(" -NumCookersToSpawn=%d"), NumCookers);
+		OptionalParams += FString::Printf(TEXT(" -NumCookersToSpawn=%d"), NumCookers); 
 	}
 
 	const bool bRunningDebug = FParse::Param(FCommandLine::Get(), TEXT("debug"));
diff --git a/Engine/Source/Editor/Persona/Private/SAnimationEditorViewport.cpp b/Engine/Source/Editor/Persona/Private/SAnimationEditorViewport.cpp
--- a/Engine/Source/Editor/Persona/Private/SAnimationEditorViewport.cpp
+++ b/Engine/Source/Editor/Persona/Private/SAnimationEditorViewport.cpp
@@ -246,16 +246,16 @@ FText SAnimationEditorViewportTabBody::GetDisplayString() const
 		{
 			return FText::Format(LOCTEXT("Previewing", "Previewing {0}"), FText::FromString(Component->GetPreviewText()));
 		}
-		else if (Component->AnimBlueprintGeneratedClass != NULL)
+		else if (Component->AnimClass != NULL)
 		{
 			const bool bWarnAboutBoneManip = !PersonaPtr.Pin()->IsModeCurrent(FPersonaModes::AnimBlueprintEditMode);
 			if (bWarnAboutBoneManip)
 			{
-				return FText::Format(LOCTEXT("PreviewingAnimBP_WarnDisabled", "Previewing {0}. \nBone manipulation is disabled in this mode. "), FText::FromString(Component->AnimBlueprintGeneratedClass->GetName()));
+				return FText::Format(LOCTEXT("PreviewingAnimBP_WarnDisabled", "Previewing {0}. \nBone manipulation is disabled in this mode. "), FText::FromString(Component->AnimClass->GetName()));
 			}
 			else
 			{
-				return FText::Format(LOCTEXT("PreviewingAnimBP", "Previewing {0}"), FText::FromString(Component->AnimBlueprintGeneratedClass->GetName()));
+				return FText::Format(LOCTEXT("PreviewingAnimBP", "Previewing {0}"), FText::FromString(Component->AnimClass->GetName()));
 			}
 		}
 		else if (Component->SkeletalMesh == NULL)
diff --git a/Engine/Source/Editor/PropertyEditor/Private/PropertyNode.cpp b/Engine/Source/Editor/PropertyEditor/Private/PropertyNode.cpp
--- a/Engine/Source/Editor/PropertyEditor/Private/PropertyNode.cpp
+++ b/Engine/Source/Editor/PropertyEditor/Private/PropertyNode.cpp
@@ -133,6 +133,9 @@ void FPropertyNode::InitNode( const FPropertyNodeInitParams& InitParams )
 		const bool bEditInline = bIsObjectOrInterface && GotReadAddresses && MyProperty->HasMetaData(TEXT("EditInline"));
 		SetNodeFlags(EPropertyNodeFlags::EditInline, bEditInline);
 
+		const bool bPersistentInstance = MyProperty->HasAllPropertyFlags(CPF_PersistentInstance);
+		SetNodeFlags(EPropertyNodeFlags::PersistentInstance, bPersistentInstance);
+
 		//Get the property max child depth
 		if (Property->HasMetaData(TEXT("MaxPropertyDepth")))
 		{
diff --git a/Engine/Source/Editor/PropertyEditor/Private/PropertyNode.h b/Engine/Source/Editor/PropertyEditor/Private/PropertyNode.h
--- a/Engine/Source/Editor/PropertyEditor/Private/PropertyNode.h
+++ b/Engine/Source/Editor/PropertyEditor/Private/PropertyNode.h
@@ -45,6 +45,8 @@ namespace EPropertyNodeFlags
 
 	const Type	SkipChildValidation				= 1 << 21; /** true if this node should skip child validation */
 
+	const Type	PersistentInstance				= 1 << 22; /** true if this node supports creating 'editinlinenew' sub-objects */
+
 	const Type 	NoFlags							= 0;
 
 };
diff --git a/Engine/Source/Editor/PropertyEditor/Private/UserInterface/PropertyEditor/SPropertyEditorEditInline.cpp b/Engine/Source/Editor/PropertyEditor/Private/UserInterface/PropertyEditor/SPropertyEditorEditInline.cpp
--- a/Engine/Source/Editor/PropertyEditor/Private/UserInterface/PropertyEditor/SPropertyEditorEditInline.cpp
+++ b/Engine/Source/Editor/PropertyEditor/Private/UserInterface/PropertyEditor/SPropertyEditorEditInline.cpp
@@ -133,7 +133,8 @@ bool SPropertyEditorEditInline::Supports( const FPropertyNode* InTreeNode, int32
 	return InTreeNode 
 		&& !InTreeNode->IsEditConst()
 		&& InTreeNode->FindObjectItemParent()
-		&& InTreeNode->HasNodeFlags(EPropertyNodeFlags::EditInline);
+		&& InTreeNode->HasNodeFlags(EPropertyNodeFlags::EditInline)
+		&& InTreeNode->HasNodeFlags(EPropertyNodeFlags::PersistentInstance);
 }
 
 bool SPropertyEditorEditInline::Supports( const TSharedRef< class FPropertyEditor >& InPropertyEditor )
diff --git a/Engine/Source/Editor/UnrealEd/Classes/CookOnTheSide/CookOnTheFlyServer.h b/Engine/Source/Editor/UnrealEd/Classes/CookOnTheSide/CookOnTheFlyServer.h
--- a/Engine/Source/Editor/UnrealEd/Classes/CookOnTheSide/CookOnTheFlyServer.h
+++ b/Engine/Source/Editor/UnrealEd/Classes/CookOnTheSide/CookOnTheFlyServer.h
@@ -662,6 +662,7 @@ public:
 		FProcHandle ProcessHandle;
 		FString ResponseFileName;
 		FString BaseResponseFileName;
+		FString ManifestFilename;
 		void* ReadPipe;
 		int32 ReturnCode;
 		bool bFinished;
@@ -717,6 +718,7 @@ private:
 		bool bForceDisableCompressedPackages;
 		bool bIsChildCooker;
 		FString ChildCookFilename;
+		FString ChildManifestFilename;
 		TSet<FName> ChildUnsolicitedPackages;
 		TArray<FChildCooker> ChildCookers;
 		TArray<FName> TargetPlatformNames;
@@ -856,10 +858,12 @@ public:
 		TArray<FString> CookCultures; 
 		TArray<FString> IniMapSections;
 		ECookByTheBookOptions CookOptions;
+		FString CookSingleAssetName;
 		FString DLCName;
 		FString CreateReleaseVersion;
 		FString BasedOnReleaseVersion;
 		FString ChildCookFileName; // if we are the child cooker 
+		FString ChildManifestFilename; // again, only if you are the child cooker
 		bool bGenerateStreamingInstallManifests; 
 		bool bGenerateDependenciesForMaps; 
 		bool bErrorOnEngineContentUse; // this is a flag for dlc, will cause the cooker to error if the dlc references engine content
@@ -1063,7 +1067,7 @@ private:
 	 * 
 	 * @param NumChildCookersToSpawn, number of child cookers we want to use
 	 */
-	void StartChildCookers(int32 NumChildCookersToSpawn, const TArray<FName>& TargetPlatformNames );
+	void StartChildCookers(int32 NumChildCookersToSpawn, const TArray<FName>& TargetPlatformNames, const FString& ExtraCmdParams = FString());
 
 	/**
 	 * TickChildCookers
diff --git a/Engine/Source/Editor/UnrealEd/Classes/Settings/EditorExperimentalSettings.h b/Engine/Source/Editor/UnrealEd/Classes/Settings/EditorExperimentalSettings.h
--- a/Engine/Source/Editor/UnrealEd/Classes/Settings/EditorExperimentalSettings.h
+++ b/Engine/Source/Editor/UnrealEd/Classes/Settings/EditorExperimentalSettings.h
@@ -69,6 +69,14 @@ public:
 	UPROPERTY(EditAnywhere, config, Category=Blueprints, meta=(DisplayName="Blueprint Break on Exceptions"))
 	bool bBreakOnExceptions;
 
+	/** Enable experimental blueprint performance analysis tools. */
+	UPROPERTY(EditAnywhere, config, Category=Blueprints, meta=(DisplayName="Blueprint Performance Analysis Tools"))
+	bool bBlueprintPerformanceAnalysisTools;
+
+	/** The number of samples the blueprint profiler should use to average the current value. */
+	UPROPERTY(EditAnywhere, config, Category=Blueprints, meta=(DisplayName="Blueprint Profiler Sample Count"))
+	int32 BlueprintProfilerAverageSampleCount;
+
 	/** Enables the visual diff tool for widget blueprints. WARNING: changes to the widget hierarchy will not be detected */
 	UPROPERTY(EditAnywhere, config, Category=Blueprints, meta=(DisplayName="Use the Diff Tool for Widget Blueprints"))
 	bool bEnableWidgetVisualDiff;
diff --git a/Engine/Source/Editor/UnrealEd/Classes/Settings/ProjectPackagingSettings.h b/Engine/Source/Editor/UnrealEd/Classes/Settings/ProjectPackagingSettings.h
--- a/Engine/Source/Editor/UnrealEd/Classes/Settings/ProjectPackagingSettings.h
+++ b/Engine/Source/Editor/UnrealEd/Classes/Settings/ProjectPackagingSettings.h
@@ -82,6 +82,10 @@ public:
 	UPROPERTY(config, EditAnywhere, Category=Project)
 	bool IncludeDebugFiles;
 
+	/** If enabled, then the project's Blueprint assets (including structs and enums) will be intermediately converted into C++ and used in the packaged project (in place of the .uasset files).*/
+	UPROPERTY(config, EditAnywhere, Category=Experimental)
+	bool bNativizeBlueprintAssets;
+
 	/** If enabled, all content will be put into a single .pak file instead of many individual files (default = enabled). */
 	UPROPERTY(config, EditAnywhere, Category=Packaging)
 	bool UsePakFile;
diff --git a/Engine/Source/Editor/UnrealEd/Private/Commandlets/CookCommandlet.cpp b/Engine/Source/Editor/UnrealEd/Private/Commandlets/CookCommandlet.cpp
--- a/Engine/Source/Editor/UnrealEd/Private/Commandlets/CookCommandlet.cpp
+++ b/Engine/Source/Editor/UnrealEd/Private/Commandlets/CookCommandlet.cpp
@@ -6,6 +6,8 @@
 
 #include "UnrealEd.h"
 
+#include "Blueprint/BlueprintSupport.h"
+#include "BlueprintNativeCodeGenModule.h"
 #include "Engine/WorldComposition.h"
 #include "PackageHelperFunctions.h"
 #include "DerivedDataCacheInterface.h"
@@ -460,7 +462,12 @@ bool UCookCommandlet::SaveCookedPackage( UPackage* Package, uint32 SaveFlags, bo
 				}
 				else
 				{
-					bSavedCorrectly &= GEditor->SavePackage( Package, World, Flags, *PlatFilename, GError, NULL, bSwap, false, SaveFlags, Target, FDateTime::MinValue() );
+					ESavePackageResult Result = GEditor->Save(Package, World, Flags, *PlatFilename, GError, NULL, bSwap, false, SaveFlags, Target, FDateTime::MinValue());
+					if (Result == ESavePackageResult::ReplaceCompletely || Result == ESavePackageResult::GenerateStub)
+					{
+						IBlueprintNativeCodeGenModule::Get().Convert(Package, Result == ESavePackageResult::ReplaceCompletely ? EReplacementResult::ReplaceCompletely : EReplacementResult::GenerateStub);
+					}
+					bSavedCorrectly &= (Result == ESavePackageResult::ReplaceCompletely || Result == ESavePackageResult::GenerateStub || Result == ESavePackageResult::Success);
 				}
 				
 				bOutWasUpToDate = false;
@@ -1055,6 +1062,9 @@ bool UCookCommandlet::NewCook( const TArray<ITargetPlatform*>& Platforms, TArray
 	FString ChildCookFile;
 	FParse::Value(*Params, TEXT("cookchild="), ChildCookFile);
 
+	FString ChildManifestFilename;
+	FParse::Value(*Params, TEXT("childmanifest="), ChildManifestFilename);
+
 	int32 NumProcesses = 0;
 	FParse::Value(*Params, TEXT("numcookerstospawn="), NumProcesses);
 
@@ -1158,11 +1168,19 @@ bool UCookCommandlet::NewCook( const TArray<ITargetPlatform*>& Platforms, TArray
 	AlwaysCookMapList.Append(MapList);
 	Swap(MapList, AlwaysCookMapList);
 
+	FCookCommandParams CookParams(FCommandLine::Get());
+	if (CookParams.bRunConversion)
+	{
+		const UBlueprintNativeCodeGenConfig* ConfigSettings = GetDefault<UBlueprintNativeCodeGenConfig>();
+		TMap<UObject*, UClass*> ClassReplacementMap;
+		ClassReplacementMap.Add(UUserDefinedEnum::StaticClass(), UEnum::StaticClass());
+		ClassReplacementMap.Add(UUserDefinedStruct::StaticClass(), UScriptStruct::StaticClass());
+		ClassReplacementMap.Add(UBlueprintGeneratedClass::StaticClass(), UDynamicClass::StaticClass());
+		FScriptCookReplacementCoordinator::Create(CookParams.bRunConversion, ConfigSettings->ExcludedAssetTypes, ConfigSettings->ExcludedBlueprintTypes, ClassReplacementMap);
+	}
+
 	//////////////////////////////////////////////////////////////////////////
 	// start cook by the book 
-
-
-
 	ECookByTheBookOptions CookOptions = ECookByTheBookOptions::None;
 
 	CookOptions |= bLeakTest ? ECookByTheBookOptions::LeakTest : ECookByTheBookOptions::None; 
@@ -1185,7 +1203,9 @@ bool UCookCommandlet::NewCook( const TArray<ITargetPlatform*>& Platforms, TArray
 	StartupOptions.bGenerateDependenciesForMaps = Switches.Contains(TEXT("GenerateDependenciesForMaps"));
 	StartupOptions.bGenerateStreamingInstallManifests = bGenerateStreamingInstallManifests;
 	StartupOptions.ChildCookFileName = ChildCookFile;
+	StartupOptions.ChildManifestFilename = ChildManifestFilename;
 	StartupOptions.NumProcesses = NumProcesses;
+	FParse::Value(FCommandLine::Get(), TEXT("COOKSINGLEASSETNAME="), StartupOptions.CookSingleAssetName);
 
 	CookOnTheFlyServer->StartCookByTheBook( StartupOptions );
 
diff --git a/Engine/Source/Editor/UnrealEd/Private/CookOnTheFlyServer.cpp b/Engine/Source/Editor/UnrealEd/Private/CookOnTheFlyServer.cpp
--- a/Engine/Source/Editor/UnrealEd/Private/CookOnTheFlyServer.cpp
+++ b/Engine/Source/Editor/UnrealEd/Private/CookOnTheFlyServer.cpp
@@ -21,6 +21,7 @@
 
 #include "AssetRegistryModule.h"
 #include "AssetData.h"
+#include "BlueprintNativeCodeGenModule.h"
 
 #include "UnrealEdMessages.h"
 #include "GameDelegates.h"
@@ -712,6 +713,12 @@ FString GetChildCookerResultFilename(const FString& ResponseFilename)
 	return Result;
 }
 
+FString GetChildCookerManifestFilename(const FString ResponseFilename)
+{
+	FString Result = ResponseFilename + TEXT("Manifest.txt");
+	return Result;
+}
+
 /**
  * Return the release asset registry filename for the release version supplied
  */
@@ -2688,6 +2695,10 @@ ESavePackageResult UCookOnTheFlyServer::SaveCookedPackage(UPackage* Package, uin
 				{
 					SCOPE_TIMER(GEditorSavePackage);
 					Result = GEditor->Save(Package, World, Flags, *PlatFilename, GError, NULL, bSwap, false, SaveFlags, Target, FDateTime::MinValue(), false);
+					if (Result == ESavePackageResult::ReplaceCompletely || Result == ESavePackageResult::GenerateStub)
+					{
+						IBlueprintNativeCodeGenModule::Get().Convert(Package, Result == ESavePackageResult::ReplaceCompletely ? EReplacementResult::ReplaceCompletely : EReplacementResult::GenerateStub);
+					}
 					INC_INT_STAT(SavedPackage, 1);
 				}
 
@@ -3885,14 +3896,21 @@ void UCookOnTheFlyServer::CookByTheBookFinished()
 			UncookedPackageList.Append(UncookedPackage.ToString() + TEXT("\n\r"));
 		}
 		FFileHelper::SaveStringToFile(UncookedPackageList, *GetChildCookerResultFilename(CookByTheBookOptions->ChildCookFilename));
-
-		
-		
+		IBlueprintNativeCodeGenModule::Get().SaveManifest(*GetChildCookerManifestFilename(CookByTheBookOptions->ChildManifestFilename));
 	}
 	else
 	{
 		CleanUpChildCookers();
 
+		IBlueprintNativeCodeGenModule& CodeGenModule = IBlueprintNativeCodeGenModule::Get();
+		// merge the manifest for the bluepritn code generator:
+		for (auto& Cooker : CookByTheBookOptions->ChildCookers)
+		{
+			CodeGenModule.MergeManifest(*GetChildCookerManifestFilename(Cooker.ManifestFilename));
+		}
+		
+		CodeGenModule.FinalizeManifest();
+
 		check(CookByTheBookOptions->ChildUnsolicitedPackages.Num() == 0);
 		SCOPE_COOKING_STAT(SavingAssetRegistry);
 		SCOPE_TIMER(SavingAssetRegistry);
@@ -4199,6 +4217,7 @@ void UCookOnTheFlyServer::StartCookByTheBook( const FCookByTheBookStartupOptions
 	CookByTheBookOptions->bForceDisableCompressedPackages = !!(CookOptions & ECookByTheBookOptions::ForceDisableCompressed);
 	CookByTheBookOptions->bIsChildCooker = CookByTheBookStartupOptions.ChildCookFileName.Len() > 0 ? true : false;
 	CookByTheBookOptions->ChildCookFilename = CookByTheBookStartupOptions.ChildCookFileName;
+	CookByTheBookOptions->ChildManifestFilename = CookByTheBookStartupOptions.ChildManifestFilename;
 	
 	NeverCookPackageList.Empty();
 	{
@@ -4375,8 +4394,14 @@ void UCookOnTheFlyServer::StartCookByTheBook( const FCookByTheBookStartupOptions
 		SaveGlobalShaderMapFiles(TargetPlatforms);
 	}
 	
-	
-	CollectFilesToCook(FilesInPath, CookMaps, CookDirectories, CookCultures, IniMapSections, bCookAll, bMapsOnly, bNoDev );
+	if (CookByTheBookStartupOptions.CookSingleAssetName.IsEmpty())
+	{
+		CollectFilesToCook(FilesInPath, CookMaps, CookDirectories, CookCultures, IniMapSections, bCookAll, bMapsOnly, bNoDev);
+	}
+	else
+	{
+		AddFileToCook(FilesInPath, CookByTheBookStartupOptions.CookSingleAssetName);
+	}
 	if (FilesInPath.Num() == 0)
 	{
 		LogCookerMessage( FString::Printf(TEXT("No files found to cook.")), EMessageSeverity::Warning );
@@ -4447,10 +4472,11 @@ void UCookOnTheFlyServer::StartCookByTheBook( const FCookByTheBookStartupOptions
 		CookRequests.EnqueueUnique( MoveTemp( PreviousRequest ) );
 	}
 	CookByTheBookOptions->PreviousCookRequests.Empty();
-
+	
 	if (CookByTheBookStartupOptions.NumProcesses)
 	{
-		StartChildCookers(CookByTheBookStartupOptions.NumProcesses, TargetPlatformNames);
+		FString ExtraCommandLine;
+		StartChildCookers(CookByTheBookStartupOptions.NumProcesses, TargetPlatformNames, ExtraCommandLine);
 	}
 }
 
@@ -4550,7 +4576,7 @@ public:
 };
 
 // sue chefs away!!
-void UCookOnTheFlyServer::StartChildCookers(int32 NumCookersToSpawn, const TArray<FName>& TargetPlatformNames)
+void UCookOnTheFlyServer::StartChildCookers(int32 NumCookersToSpawn, const TArray<FName>& TargetPlatformNames, const FString& ExtraCmdParams)
 {
 	SCOPE_TIMER(StartingChildCookers);
 	// create a comprehensive list of all the files we need to cook
@@ -4661,11 +4687,10 @@ void UCookOnTheFlyServer::StartChildCookers(int32 NumCookersToSpawn, const TArra
 
 		ChildCooker.ResponseFileName = FPaths::CreateTempFilename(*(FPaths::GameSavedDir() / TEXT("CookingTemp")));
 		ChildCooker.BaseResponseFileName = FPaths::GetBaseFilename(ChildCooker.ResponseFileName);
+		ChildCooker.ManifestFilename = FString(*(FPaths::GameSavedDir() / TEXT("CookingTemp") / *FString::Printf(TEXT("BPCodeManifest_%d.json"), CookerCounter)));
 		// FArchive* ResponseFile = IFileManager::CreateFileWriter(ChildCooker.UniqueTempName);
 		FString ResponseFileText;
 
-
-
 		for (int32 I = 0; I < NumFilesForCooker; ++I)
 		{
 			FName PackageFName = DistributeStandardFilenames[I];
@@ -4698,7 +4723,8 @@ void UCookOnTheFlyServer::StartChildCookers(int32 NumCookersToSpawn, const TArra
 
 		// default commands
 		// multiprocess tells unreal in general we shouldn't do things like save ddc, clean shader working directory, and other various multiprocess unsafe things
-		FString CommandLine = FString::Printf(TEXT("\"%s\" -run=cook -multiprocess -targetplatform=%s -cookchild=\"%s\" -abslog=\"%sLog.txt\""), *FPaths::GetProjectFilePath(), *TargetPlatformString, *ChildCooker.ResponseFileName, *ChildCooker.ResponseFileName);
+		FString CommandLine = FString::Printf(TEXT("\"%s\" -run=cook -multiprocess -targetplatform=%s -cookchild=\"%s\" -abslog=\"%sLog.txt\" -childmanifest=\"%s\" %s"), 
+			*FPaths::GetProjectFilePath(), *TargetPlatformString, *ChildCooker.ResponseFileName, *ChildCooker.ResponseFileName, *ChildCooker.ManifestFilename, *ExtraCmdParams);
 
 		auto KeepCommandlineValue = [&](const TCHAR* CommandlineToKeep)
 		{
diff --git a/Engine/Source/Editor/UnrealEd/Private/Factories/ActorFactory.cpp b/Engine/Source/Editor/UnrealEd/Private/Factories/ActorFactory.cpp
--- a/Engine/Source/Editor/UnrealEd/Private/Factories/ActorFactory.cpp
+++ b/Engine/Source/Editor/UnrealEd/Private/Factories/ActorFactory.cpp
@@ -1142,7 +1142,7 @@ void UActorFactorySkeletalMesh::PostCreateBlueprint( UObject* Asset, AActor* CDO
 
 		ASkeletalMeshActor* SkeletalMeshActor = CastChecked<ASkeletalMeshActor>(CDO);
 		SkeletalMeshActor->GetSkeletalMeshComponent()->SkeletalMesh = SkeletalMesh;
-		SkeletalMeshActor->GetSkeletalMeshComponent()->AnimBlueprintGeneratedClass = AnimBlueprint ? Cast<UAnimBlueprintGeneratedClass>(AnimBlueprint->GeneratedClass) : NULL;
+		SkeletalMeshActor->GetSkeletalMeshComponent()->AnimClass = AnimBlueprint ? Cast<UAnimBlueprintGeneratedClass>(AnimBlueprint->GeneratedClass) : NULL;
 	}
 }
 
diff --git a/Engine/Source/Editor/UnrealEd/Private/Kismet2/BlueprintEditorUtils.cpp b/Engine/Source/Editor/UnrealEd/Private/Kismet2/BlueprintEditorUtils.cpp
--- a/Engine/Source/Editor/UnrealEd/Private/Kismet2/BlueprintEditorUtils.cpp
+++ b/Engine/Source/Editor/UnrealEd/Private/Kismet2/BlueprintEditorUtils.cpp
@@ -342,6 +342,13 @@ void FBasePinChangeHelper::Broadcast(UBlueprint* InBlueprint, UK2Node_EditablePi
 			if (NodeIsNotTransient(CallSite))
 			{
 				UBlueprint* CallSiteBlueprint = FBlueprintEditorUtils::FindBlueprintForNode(CallSite);
+				if (CallSiteBlueprint == nullptr)
+				{
+					// the node doesn't have a Blueprint in its outer chain, 
+					// probably signifying that it is part of a graph that has 
+					// been removed by the user (and moved off the Blueprint)
+					continue;
+				}
 
 				const bool bNameMatches = (CallSite->FunctionReference.GetMemberName() == FuncName);
 				const UClass* MemberParentClass = CallSite->FunctionReference.GetMemberParentClass(CallSite->GetBlueprintClassFromNode());
@@ -1886,6 +1893,8 @@ void FBlueprintEditorUtils::PropagateParentBlueprintDefaults(UClass* ClassToProp
 
 UNREALED_API FSecondsCounterData BlueprintCompileAndLoadTimerData;
 
+uint32 FBlueprintDuplicationScopeFlags::bStaticFlags = FBlueprintDuplicationScopeFlags::NoFlags;
+
 void FBlueprintEditorUtils::PostDuplicateBlueprint(UBlueprint* Blueprint, bool bDuplicateForPIE)
 {
 	FSecondsCounterScope Timer(BlueprintCompileAndLoadTimerData); 
@@ -1985,6 +1994,11 @@ void FBlueprintEditorUtils::PostDuplicateBlueprint(UBlueprint* Blueprint, bool b
 				UTimelineTemplate* OldTimeline = *TimelineIt;
 				UTimelineTemplate* NewTimeline = CastChecked<UTimelineTemplate>(StaticDuplicateObject(OldTimeline, NewBPGC, OldTimeline->GetFName()));
 
+				if (FBlueprintDuplicationScopeFlags::HasAnyFlag(FBlueprintDuplicationScopeFlags::TheSameTimelineGuid))
+				{
+					NewTimeline->TimelineGuid = OldTimeline->TimelineGuid;
+				}
+
 				NewBPGC->Timelines.Add(NewTimeline);
 				OldToNewMap.Add(OldTimeline, NewTimeline);
 			}
@@ -2056,8 +2070,11 @@ void FBlueprintEditorUtils::PostDuplicateBlueprint(UBlueprint* Blueprint, bool b
 			UEditorEngine::CopyPropertiesForUnrelatedObjects(OldCDO, NewCDO);
 		}
 
-		// And compile again to make sure they go into the generated class, get cleaned up, etc...
-		FKismetEditorUtilities::CompileBlueprint(Blueprint, false, true);
+		if (!FBlueprintDuplicationScopeFlags::HasAnyFlag(FBlueprintDuplicationScopeFlags::NoExtraCompilation))
+		{
+			// And compile again to make sure they go into the generated class, get cleaned up, etc...
+			FKismetEditorUtilities::CompileBlueprint(Blueprint, false, true);
+		}
 
 		// it can still keeps references to some external objects
 		Blueprint->LastEditedDocuments.Empty();
@@ -2823,6 +2840,13 @@ void FBlueprintEditorUtils::GatherDependencies(const UBlueprint* InBlueprint, TS
 	Dependencies.Empty();
 	OutUDSDependencies.Empty();
 
+	// If the Blueprint's GeneratedClass was not generated by the Blueprint, it's either corrupt or a PIE version of the BP
+	if (InBlueprint->GeneratedClass && InBlueprint->GeneratedClass->ClassGeneratedBy != InBlueprint)
+	{
+		// Dependencies do not matter for PIE duplicated Blueprints
+		return;
+	}
+
 	InBlueprint->GatherDependencies(Dependencies);
 
 	FGatherDependenciesHelper::ProcessHierarchy(InBlueprint->ParentClass, Dependencies);
@@ -4679,11 +4703,14 @@ bool FBlueprintEditorUtils::AddLocalVariable(UBlueprint* Blueprint, UEdGraph* In
 {
 	if(InTargetGraph != NULL && InTargetGraph->GetSchema()->GetGraphType(InTargetGraph) == GT_Function)
 	{
+		// Ensure we have the top level graph for the function, in-case we are in a child graph
+		UEdGraph* TargetGraph = FBlueprintEditorUtils::GetTopLevelGraph(InTargetGraph);
+
 		const FScopedTransaction Transaction( LOCTEXT("AddLocalVariable", "Add Local Variable") );
 		Blueprint->Modify();
 
 		TArray<UK2Node_FunctionEntry*> FunctionEntryNodes;
-		InTargetGraph->GetNodesOfClass(FunctionEntryNodes);
+		TargetGraph->GetNodesOfClass(FunctionEntryNodes);
 		check(FunctionEntryNodes.Num());
 
 		const UEdGraphSchema_K2* K2Schema = GetDefault<UEdGraphSchema_K2>();
@@ -4709,7 +4736,7 @@ bool FBlueprintEditorUtils::AddLocalVariable(UBlueprint* Blueprint, UEdGraph* In
 		{
 			// Pull the description of the new variable, we will update it's default value with the default value from the user defined struct.
 			FBPVariableDescription* VariableDescription = &FunctionEntryNodes[0]->LocalVariables[FunctionEntryNodes[0]->LocalVariables.Num() - 1];
-			SetDefaultValueOnUserDefinedStructProperty(Blueprint, InTargetGraph->GetName(), VariableDescription);
+			SetDefaultValueOnUserDefinedStructProperty(Blueprint, TargetGraph->GetName(), VariableDescription);
 		}
 
 		return true;
@@ -5968,10 +5995,21 @@ void FBlueprintEditorUtils::ConformImplementedInterfaces(UBlueprint* Blueprint)
 	Blueprint->GetAllGraphs(AllGraphs);
 	// collect all existing event nodes, so we can find interface events that 
 	// need to be converted to function graphs
-	TArray<UK2Node_Event*> ImplementedEvents;
+	TArray<UK2Node_Event*> PotentialInterfaceEvents;
 	for (UEdGraph const* Graph : AllGraphs)
 	{
-		Graph->GetNodesOfClass<UK2Node_Event>(ImplementedEvents);
+		UClass* InterfaceEventClass = UK2Node_Event::StaticClass();
+		for (UEdGraphNode* GraphNode : Graph->Nodes)
+		{
+			// interface event nodes are only ever going to be implemented as 
+			// explicit UK2Node_Events... using == instead of IsChildOf<> 
+			// guarantees that we won't be catching any special node types that
+			// users might have made (that maybe reference interface functions too)
+			if (GraphNode && (GraphNode->GetClass() == InterfaceEventClass))
+			{
+				PotentialInterfaceEvents.Add(CastChecked<UK2Node_Event>(GraphNode));
+			}
+		}
 	}
 
 	for (int32 InterfaceIndex = 0; InterfaceIndex < Blueprint->ImplementedInterfaces.Num(); )
@@ -5991,7 +6029,7 @@ void FBlueprintEditorUtils::ConformImplementedInterfaces(UBlueprint* Blueprint)
 		}
 
 		// now try to conform by name/signature
-		ConformInterfaceByName(Blueprint, CurrentInterface, InterfaceIndex, ImplementedEvents, VariableNamesUsedInBlueprint);
+		ConformInterfaceByName(Blueprint, CurrentInterface, InterfaceIndex, PotentialInterfaceEvents, VariableNamesUsedInBlueprint);
 
 		// not going to remove this interface, so let's continue forward
 		++InterfaceIndex;
@@ -7227,7 +7265,7 @@ void FBlueprintEditorUtils::UpdateOldPureFunctions( UBlueprint* Blueprint )
 						{
 							UK2Node_FunctionEntry* Entry = EntryNodes[0];
 							Entry->Modify();
-							Entry->SetExtraFlags(Entry->GetFunctionFlags() | FUNC_BlueprintPure);
+							Entry->AddExtraFlags(FUNC_BlueprintPure);
 						}
 					}
 				}
