Fix crash when trying to use fast path on BreakTransform (by stopping BreakTransform from using fast path copying)

diff --git a/Engine/Source/Editor/KismetCompiler/Private/AnimBlueprintCompiler.cpp b/Engine/Source/Editor/KismetCompiler/Private/AnimBlueprintCompiler.cpp
--- a/Engine/Source/Editor/KismetCompiler/Private/AnimBlueprintCompiler.cpp
+++ b/Engine/Source/Editor/KismetCompiler/Private/AnimBlueprintCompiler.cpp
@@ -2077,6 +2077,7 @@ bool FAnimBlueprintCompiler::FEvaluationHandlerRecord::CheckForLogicalNot(FAnimN
 
 bool FAnimBlueprintCompiler::FEvaluationHandlerRecord::CheckForStructMemberAccess(FAnimNodeSinglePropertyHandler& Handler, UEdGraphPin* SourcePin)
 {
+	static const FName BreakTransformName(TEXT("BreakTransform"));
 	if(SourcePin)
 	{
 		UEdGraphPin* DestPin = nullptr;
@@ -2096,7 +2097,9 @@ bool FAnimBlueprintCompiler::FEvaluationHandlerRecord::CheckForStructMemberAcces
 		else if(UK2Node_CallFunction* NativeBreakNode = Cast<UK2Node_CallFunction>(FollowKnots(SourcePin, DestPin)))
 		{
 			UFunction* Function = NativeBreakNode->FunctionReference.ResolveMember<UFunction>(UKismetMathLibrary::StaticClass());
-			if(Function && Function->HasMetaData(TEXT("NativeBreakFunc")))
+			if( Function && Function->HasMetaData(TEXT("NativeBreakFunc")) &&
+				Function->GetFName() != BreakTransformName) // Skip Break Transform as it is not compatible (it is not a pure "break" function as it performs type 
+															// conversion on the "Rotation" component. This means we cannot do a fast path copy as src/dest copy types do not match)
 			{
 				if(UEdGraphPin* InputPin = FindFirstInputPin(NativeBreakNode))
 				{
