Manual reproduction of 2837173 in Main. Would have shown up as a blocker in integration testing if we had tested against current content rather than doing testing in 4.11 (which has not integrated from main in quite some time)

diff --git a/Engine/Source/Editor/KismetCompiler/Private/AnimBlueprintCompiler.cpp b/Engine/Source/Editor/KismetCompiler/Private/AnimBlueprintCompiler.cpp
--- a/Engine/Source/Editor/KismetCompiler/Private/AnimBlueprintCompiler.cpp
+++ b/Engine/Source/Editor/KismetCompiler/Private/AnimBlueprintCompiler.cpp
@@ -83,7 +83,6 @@ bool FAnimBlueprintCompiler::FEffectiveConstantRecord::Apply(UObject* Object)
 
 FAnimBlueprintCompiler::FAnimBlueprintCompiler(UAnimBlueprint* SourceSketch, FCompilerResultsLog& InMessageLog, const FKismetCompilerOptions& InCompileOptions, TArray<UObject*>* InObjLoaded)
 	: FKismetCompilerContext(SourceSketch, InMessageLog, InCompileOptions, InObjLoaded)
-	, NewAnimBlueprintClass(nullptr)
 	, AnimBlueprint(SourceSketch)
 	, bIsDerivedAnimBlueprint(false)
 {
@@ -104,12 +103,6 @@ FAnimBlueprintCompiler::~FAnimBlueprintCompiler()
 {
 }
 
-void FAnimBlueprintCompiler::SetClassForBytecodeCompile(UClass* TargetClass)
-{
-	FKismetCompilerContext::SetClassForBytecodeCompile(TargetClass);
-	NewAnimBlueprintClass = Cast<UAnimBlueprintGeneratedClass>(TargetClass);
-}
-
 UEdGraphSchema_K2* FAnimBlueprintCompiler::CreateSchema()
 {
 	AnimSchema = NewObject<UAnimationGraphSchema>();
diff --git a/Engine/Source/Editor/KismetCompiler/Private/AnimBlueprintCompiler.h b/Engine/Source/Editor/KismetCompiler/Private/AnimBlueprintCompiler.h
--- a/Engine/Source/Editor/KismetCompiler/Private/AnimBlueprintCompiler.h
+++ b/Engine/Source/Editor/KismetCompiler/Private/AnimBlueprintCompiler.h
@@ -30,7 +30,6 @@ public:
 
 protected:
 	// Implementation of FKismetCompilerContext interface
-	virtual void SetClassForBytecodeCompile(UClass* TargetClass) override;
 	virtual UEdGraphSchema_K2* CreateSchema() override;
 	virtual void MergeUbergraphPagesIn(UEdGraph* Ubergraph) override;
 	virtual void ProcessOneFunctionGraph(UEdGraph* SourceGraph, bool bInternalFunction = false) override;
diff --git a/Engine/Source/Editor/KismetCompiler/Private/KismetCompiler.cpp b/Engine/Source/Editor/KismetCompiler/Private/KismetCompiler.cpp
--- a/Engine/Source/Editor/KismetCompiler/Private/KismetCompiler.cpp
+++ b/Engine/Source/Editor/KismetCompiler/Private/KismetCompiler.cpp
@@ -95,11 +95,6 @@ FKismetCompilerContext::~FKismetCompilerContext()
 	DefaultPropertyValueMap.Empty();
 }
 
-void FKismetCompilerContext::SetClassForBytecodeCompile(UClass* TargetClass)
-{
-	NewClass = Cast<UBlueprintGeneratedClass>(TargetClass);
-}
-
 UEdGraphSchema_K2* FKismetCompilerContext::CreateSchema()
 {
 	return NewObject<UEdGraphSchema_K2>();
@@ -417,13 +412,6 @@ UProperty* FKismetCompilerContext::CreateVariable(const FName VarName, const FEd
 			*VarName.ToString()), Blueprint);
 	}
 
-	if (CompileOptions.CompileType == EKismetCompileType::BytecodeOnly)
-	{
-		UProperty* Ret = NewClass->FindPropertyByName(VarName);
-		check(Ret);
-		return Ret;
-	}
-
 	UProperty* NewProperty = FKismetCompilerUtilities::CreatePropertyOnScope(NewClass, VarName, VarType, NewClass, 0, Schema, MessageLog);
 	if (NewProperty != NULL)
 	{
@@ -666,17 +654,7 @@ void FKismetCompilerContext::CreatePropertiesFromList(UStruct* Scope, UField**&
 			MessageLog.Error(*FString::Printf(*LOCTEXT("PropertyForLiteral_Error", "Cannot create property for a literal: %s from @@ type (%s)").ToString(), *Term.Name, *UEdGraphSchema_K2::TypeToText(Term.Type).ToString()), Term.Source);
 		}
 
-		UProperty* NewProperty = nullptr;
-		if (CompileOptions.CompileType == EKismetCompileType::BytecodeOnly)
-		{
-			NewProperty = Scope->FindPropertyByName(FName(*Term.Name));
-			check(NewProperty);
-		}
-		else
-		{
-			NewProperty = FKismetCompilerUtilities::CreatePropertyOnScope(Scope, FName(*Term.Name), Term.Type, NewClass, PropertyFlags, Schema, MessageLog);
-		}
-
+		UProperty* NewProperty = FKismetCompilerUtilities::CreatePropertyOnScope(Scope, FName(*Term.Name), Term.Type, NewClass, PropertyFlags, Schema, MessageLog);
 		if (NewProperty != NULL)
 		{
 			NewProperty->PropertyFlags |= PropertyFlags;
@@ -735,11 +713,8 @@ void FKismetCompilerContext::CreatePropertiesFromList(UStruct* Scope, UField**&
 			}
 			
 			// Link this object to the tail of the list (so properties remain in the desired order)
-			if (CompileOptions.CompileType != EKismetCompileType::BytecodeOnly)
-			{
-				*PropertyStorageLocation = NewProperty;
-				PropertyStorageLocation = &(NewProperty->Next);
-			}
+			*PropertyStorageLocation = NewProperty;
+			PropertyStorageLocation = &(NewProperty->Next);
 
 			Term.AssociatedVarProperty = NewProperty;
 			Term.SetVarTypeLocal(bPropertiesAreLocal);
@@ -807,12 +782,9 @@ void FKismetCompilerContext::CreateLocalVariablesForFunction(FKismetFunctionCont
 		}
 
 		// If there were local properties, place them at the end of the property storage location
-		if (CompileOptions.CompileType != EKismetCompileType::BytecodeOnly)
+		if(LocalProperties)
 		{
-			if (LocalProperties)
-			{
-				*PropertyStorageLocation = LocalProperties;
-			}
+			*PropertyStorageLocation = LocalProperties;
 		}
 
 		// Create debug data for variable reads/writes
@@ -1228,7 +1200,7 @@ void FKismetCompilerContext::PrecompileFunction(FKismetFunctionContext& Context)
 			CreatedFunctionNames.Add(NewFunctionNameString);
 		}
 
-		Context.Function = CompileOptions.CompileType == EKismetCompileType::BytecodeOnly ? NewClass->FindFunctionByName(NewFunctionName) : NewObject<UFunction>(NewClass, NewFunctionName, RF_Public);
+		Context.Function = NewObject<UFunction>(NewClass, NewFunctionName, RF_Public);
 
 #if USE_TRANSIENT_SKELETON
 		// Propagate down transient settings from the class
@@ -1326,30 +1298,22 @@ void FKismetCompilerContext::PrecompileFunction(FKismetFunctionContext& Context)
 			}
 		}
 		
+		// Link it
+		//@TODO: should this be in regular or reverse order?
+		Context.Function->Next = Context.NewClass->Children;
+		Context.NewClass->Children = Context.Function;
 
-		if (CompileOptions.CompileType != EKismetCompileType::BytecodeOnly)
-		{
-			// Link it
-			//@TODO: should this be in regular or reverse order?
-			Context.Function->Next = Context.NewClass->Children;
-			Context.NewClass->Children = Context.Function;
-
-			// Add the function to it's owner class function name -> function map
-			Context.NewClass->AddFunctionToFunctionMap(Context.Function);
-		}
-
+		// Add the function to it's owner class function name -> function map
+		Context.NewClass->AddFunctionToFunctionMap(Context.Function);
 		if (UsePersistentUberGraphFrame() && Context.bIsUbergraph)
 		{
-			ensure(!NewClass->UberGraphFunction || (NewClass->UberGraphFunction == Context.Function && CompileOptions.CompileType == EKismetCompileType::BytecodeOnly));
+			ensure(!NewClass->UberGraphFunction);
 			NewClass->UberGraphFunction = Context.Function;
 		}
 
 		// Create any user defined variables, this must occur before registering nets so that the properties are in place
-		if (CompileOptions.CompileType != EKismetCompileType::BytecodeOnly)
-		{
-			UField** PropertyStorageLocation = &(Context.Function->Children);
-			CreateUserDefinedLocalVariablesForFunction(Context, PropertyStorageLocation);
-		}
+		UField** PropertyStorageLocation = &(Context.Function->Children);
+		CreateUserDefinedLocalVariablesForFunction(Context, PropertyStorageLocation);
 
 		//@TODO: Prune pure functions that don't have any consumers
 		if (bIsFullCompile)
@@ -1408,11 +1372,8 @@ void FKismetCompilerContext::PrecompileFunction(FKismetFunctionContext& Context)
 
 		Context.Function->FunctionFlags |= Context.GetNetFlags();
 
-		if (CompileOptions.CompileType != EKismetCompileType::BytecodeOnly)
-		{
-			// Parameter list needs to be linked before signatures are compared. 
-			Context.Function->StaticLink(true);
-		}
+		// Parameter list needs to be linked before signatures are compared. 
+		Context.Function->StaticLink(true);
 
 		// Make sure the function signature is valid if this is an override
 		if (ParentFunction)
@@ -1647,12 +1608,8 @@ void FKismetCompilerContext::PostcompileFunction(FKismetFunctionContext& Context
 void FKismetCompilerContext::FinishCompilingFunction(FKismetFunctionContext& Context)
 {
 	UFunction* Function = Context.Function;
-
-	if (CompileOptions.CompileType != EKismetCompileType::BytecodeOnly)
-	{
-		Function->Bind();
-		Function->StaticLink(true);
-	}
+	Function->Bind();
+	Function->StaticLink(true);
 
 	// Set function flags and calculate cached values so the class can be used immediately
 	Function->ParmsSize = 0;
@@ -3014,44 +2971,9 @@ void FKismetCompilerContext::CreateFunctionList()
 	BP_SCOPED_COMPILER_EVENT_STAT(EKismetCompilerStats_CreateFunctionList);
 
 	// Process the ubergraph if one should be present
-	if (FBlueprintEditorUtils::DoesSupportEventGraphs(Blueprint) )
+	if (FBlueprintEditorUtils::DoesSupportEventGraphs(Blueprint))
 	{
-		if (CompileOptions.CompileType != EKismetCompileType::BytecodeOnly)
-		{
-			CreateAndProcessUbergraph();
-			Blueprint->ConsolidatedEventGraph = ConsolidatedEventGraph;
-		}
-		else
-		{
-			if (Blueprint->ConsolidatedEventGraph && Blueprint->ConsolidatedEventGraph->Nodes.Num() > 0)
-			{
-				UbergraphContext = new (FunctionList)FKismetFunctionContext(MessageLog, Schema, NewClass, Blueprint, CompileOptions.DoesRequireCppCodeGeneration(), CompileOptions.IsInstrumentationActive());
-				UbergraphContext->SourceGraph = Blueprint->ConsolidatedEventGraph;
-				ConsolidatedEventGraph = UbergraphContext->SourceGraph;
-				Blueprint->ConsolidatedEventGraph = nullptr;
-				UbergraphContext->MarkAsEventGraph();
-				UbergraphContext->MarkAsInternalOrCppUseOnly();
-				UbergraphContext->SetExternalNetNameMap(&ClassScopeNetNameMap);
-
-				// Validate all the nodes in the graph
-				for (int32 ChildIndex = 0; ChildIndex < ConsolidatedEventGraph->Nodes.Num(); ++ChildIndex)
-				{
-					const UEdGraphNode* Node = ConsolidatedEventGraph->Nodes[ChildIndex];
-					const int32 SavedErrorCount = MessageLog.NumErrors;
-					UK2Node_Event* SrcEventNode = Cast<UK2Node_Event>(ConsolidatedEventGraph->Nodes[ChildIndex]);
-					if (bIsFullCompile || SrcEventNode)
-					{
-						ValidateNode(Node);
-					}
-
-					// If the node didn't generate any errors then generate function stubs for event entry nodes etc.
-					if ((SavedErrorCount == MessageLog.NumErrors) && SrcEventNode)
-					{
-						CreateFunctionStubForEvent(SrcEventNode, Blueprint);
-					}
-				}
-			}
-		}
+		CreateAndProcessUbergraph();
 	}
 
 	if (Blueprint->BlueprintType != BPTYPE_MacroLibrary)
@@ -3256,27 +3178,16 @@ void FKismetCompilerContext::Compile()
 		++TimelineIndex;
 	}
 
-	if (CompileOptions.CompileType != EKismetCompileType::BytecodeOnly)
-	{
-		CleanAndSanitizeClass(TargetClass, OldCDO);
-	}
-	else
-	{
-		SetClassForBytecodeCompile(TargetClass);
-		OldCDO = TargetClass->ClassDefaultObject;
-	}
+	CleanAndSanitizeClass(TargetClass, OldCDO);
 
 	FKismetCompilerVMBackend Backend_VM(Blueprint, Schema, *this);
 
 	NewClass->ClassGeneratedBy = Blueprint;
 
-	if (CompileOptions.CompileType != EKismetCompileType::BytecodeOnly)
-	{
-		UClass* ParentClass = NewClass->ClassWithin;
-		NewClass->SetSuperStruct(ParentClass);
-		NewClass->ClassFlags |= (ParentClass->ClassFlags & CLASS_Inherit);
-		NewClass->ClassCastFlags |= ParentClass->ClassCastFlags;
-	}
+	UClass* ParentClass = NewClass->ClassWithin;
+	NewClass->SetSuperStruct(ParentClass);
+	NewClass->ClassFlags |= (ParentClass->ClassFlags & CLASS_Inherit);
+	NewClass->ClassCastFlags |= ParentClass->ClassCastFlags;
 	
 	if(Blueprint->bGenerateConstClass)
 	{
@@ -3329,7 +3240,7 @@ void FKismetCompilerContext::Compile()
 		}
 	}
 
-	if (UsePersistentUberGraphFrame() && UbergraphContext && CompileOptions.CompileType != EKismetCompileType::BytecodeOnly)
+	if (UsePersistentUberGraphFrame() && UbergraphContext)
 	{
 		//UBER GRAPH PERSISTENT FRAME
 		FEdGraphPinType Type(TEXT("struct"), TEXT(""), FPointerToUberGraphFrame::StaticStruct(), false, false);
@@ -3337,7 +3248,6 @@ void FKismetCompilerContext::Compile()
 		Property->SetPropertyFlags(CPF_DuplicateTransient | CPF_Transient);
 	}
 
-	if ( CompileOptions.CompileType != EKismetCompileType::BytecodeOnly )
 	{ BP_SCOPED_COMPILER_EVENT_STAT(EKismetCompilerStats_BindAndLinkClass);
 
 		// Relink the class
@@ -3455,13 +3365,7 @@ void FKismetCompilerContext::Compile()
 		FUserDefinedStructureCompilerUtils::DefaultUserDefinedStructs(NewCDO, MessageLog);
 
 		// Copy over the CDO properties if we're not already regenerating on load.  In that case, the copy will be done after compile on load is complete
-		// We can't do this during the bytecode pass because we don't have an 'old cdo' to copy our overwritten values from (IE if a derived type changes
-		// some default values, we will stomp them with the parent's values.. normally we can copy the previous CDO's values over, but there
-		// is no previous CDO if NewClass == OldClass)
-		if (CompileOptions.CompileType != EKismetCompileType::BytecodeOnly)
-		{
-			FBlueprintEditorUtils::PropagateParentBlueprintDefaults(NewClass);
-		}
+		FBlueprintEditorUtils::PropagateParentBlueprintDefaults(NewClass);
 
 		if (Blueprint->HasAnyFlags(RF_BeingRegenerated))
 		{
@@ -3471,8 +3375,7 @@ void FKismetCompilerContext::Compile()
 				Blueprint->PRIVATE_InnermostPreviousCDO = OldCDO;
 			}
 		}
-
-		if (CompileOptions.CompileType != EKismetCompileType::BytecodeOnly)
+		else
 		{
 			if( NewCDO )
 			{
diff --git a/Engine/Source/Editor/KismetCompiler/Public/KismetCompiler.h b/Engine/Source/Editor/KismetCompiler/Public/KismetCompiler.h
--- a/Engine/Source/Editor/KismetCompiler/Public/KismetCompiler.h
+++ b/Engine/Source/Editor/KismetCompiler/Public/KismetCompiler.h
@@ -145,7 +145,6 @@ public:
 
 	FString GetGuid(const UEdGraphNode* Node) const;
 protected:
-	virtual void SetClassForBytecodeCompile(UClass* TargetClass);
 	virtual UEdGraphSchema_K2* CreateSchema();
 	virtual void PostCreateSchema();
 	virtual void SpawnNewClass(const FString& NewClassName);
diff --git a/Engine/Source/Editor/UMGEditor/Private/WidgetBlueprintCompiler.cpp b/Engine/Source/Editor/UMGEditor/Private/WidgetBlueprintCompiler.cpp
--- a/Engine/Source/Editor/UMGEditor/Private/WidgetBlueprintCompiler.cpp
+++ b/Engine/Source/Editor/UMGEditor/Private/WidgetBlueprintCompiler.cpp
@@ -17,7 +17,6 @@
 
 FWidgetBlueprintCompiler::FWidgetBlueprintCompiler(UWidgetBlueprint* SourceSketch, FCompilerResultsLog& InMessageLog, const FKismetCompilerOptions& InCompilerOptions, TArray<UObject*>* InObjLoaded)
 	: Super(SourceSketch, InMessageLog, InCompilerOptions, InObjLoaded)
-	, NewWidgetBlueprintClass(nullptr)
 {
 }
 
@@ -25,87 +24,75 @@ FWidgetBlueprintCompiler::~FWidgetBlueprintCompiler()
 {
 }
 
-void FWidgetBlueprintCompiler::SetClassForBytecodeCompile(UClass* TargetClass)
-{
-	FKismetCompilerContext::SetClassForBytecodeCompile(TargetClass);
-
-	NewWidgetBlueprintClass = CastChecked<UWidgetBlueprintGeneratedClass>(TargetClass);
-}
-
 void FWidgetBlueprintCompiler::CreateFunctionList()
 {
-	// Create functions for bindings, unfortunately the bytecode compile pass needlessly calls this function because
-	// we need to regenerate function contexts, but we shouldn't regenerate the graph that we will be generating
-	// function contexts from, so only run the following if we're not running the bytecode pass:
-	if (CompileOptions.CompileType != EKismetCompileType::BytecodeOnly)
+	Super::CreateFunctionList();
+
+	for ( FDelegateEditorBinding& EditorBinding : WidgetBlueprint()->Bindings )
 	{
-		for (FDelegateEditorBinding& EditorBinding : WidgetBlueprint()->Bindings)
+		if ( EditorBinding.SourcePath.IsEmpty() )
 		{
-			if (EditorBinding.SourcePath.IsEmpty())
+			const FName PropertyName = EditorBinding.SourceProperty;
+
+			UProperty* Property = FindField<UProperty>(Blueprint->SkeletonGeneratedClass, PropertyName);
+			if ( Property )
 			{
-				const FName PropertyName = EditorBinding.SourceProperty;
+				// Create the function graph.
+				FString FunctionName = FString(TEXT("__Get")) + PropertyName.ToString();
+				UEdGraph* FunctionGraph = FBlueprintEditorUtils::CreateNewGraph(Blueprint, FBlueprintEditorUtils::FindUniqueKismetName(Blueprint, FunctionName), UEdGraph::StaticClass(), UEdGraphSchema_K2::StaticClass());
 
-				UProperty* Property = FindField<UProperty>(Blueprint->SkeletonGeneratedClass, PropertyName);
-				if (Property)
-				{
-					// Create the function graph.
-					FString FunctionName = FString(TEXT("__Get")) + PropertyName.ToString();
-					UEdGraph* FunctionGraph = nullptr;
-					FunctionGraph = FBlueprintEditorUtils::CreateNewGraph(Blueprint, FBlueprintEditorUtils::FindUniqueKismetName(Blueprint, FunctionName), UEdGraph::StaticClass(), UEdGraphSchema_K2::StaticClass());
+				// Update the function binding to match the generated graph name
+				EditorBinding.FunctionName = FunctionGraph->GetFName();
 
-					// Update the function binding to match the generated graph name
-					EditorBinding.FunctionName = FunctionGraph->GetFName();
+				const UEdGraphSchema_K2* K2Schema = Cast<const UEdGraphSchema_K2>(FunctionGraph->GetSchema());
 
-					const UEdGraphSchema_K2* K2Schema = Cast<const UEdGraphSchema_K2>(FunctionGraph->GetSchema());
+				Schema->CreateDefaultNodesForGraph(*FunctionGraph);
 
-					Schema->CreateDefaultNodesForGraph(*FunctionGraph);
+				K2Schema->MarkFunctionEntryAsEditable(FunctionGraph, true);
 
-					K2Schema->MarkFunctionEntryAsEditable(FunctionGraph, true);
+				// Create a function entry node
+				FGraphNodeCreator<UK2Node_FunctionEntry> FunctionEntryCreator(*FunctionGraph);
+				UK2Node_FunctionEntry* EntryNode = FunctionEntryCreator.CreateNode();
+				EntryNode->SignatureClass = NULL;
+				EntryNode->SignatureName = FunctionGraph->GetFName();
+				FunctionEntryCreator.Finalize();
 
-					// Create a function entry node
-					FGraphNodeCreator<UK2Node_FunctionEntry> FunctionEntryCreator(*FunctionGraph);
-					UK2Node_FunctionEntry* EntryNode = FunctionEntryCreator.CreateNode();
-					EntryNode->SignatureClass = NULL;
-					EntryNode->SignatureName = FunctionGraph->GetFName();
-					FunctionEntryCreator.Finalize();
+				FGraphNodeCreator<UK2Node_FunctionResult> FunctionReturnCreator(*FunctionGraph);
+				UK2Node_FunctionResult* ReturnNode = FunctionReturnCreator.CreateNode();
+				ReturnNode->SignatureClass = NULL;
+				ReturnNode->SignatureName = FunctionGraph->GetFName();
+				ReturnNode->NodePosX = EntryNode->NodePosX + EntryNode->NodeWidth + 256;
+				ReturnNode->NodePosY = EntryNode->NodePosY;
+				FunctionReturnCreator.Finalize();
 
-					FGraphNodeCreator<UK2Node_FunctionResult> FunctionReturnCreator(*FunctionGraph);
-					UK2Node_FunctionResult* ReturnNode = FunctionReturnCreator.CreateNode();
-					ReturnNode->SignatureClass = NULL;
-					ReturnNode->SignatureName = FunctionGraph->GetFName();
-					ReturnNode->NodePosX = EntryNode->NodePosX + EntryNode->NodeWidth + 256;
-					ReturnNode->NodePosY = EntryNode->NodePosY;
-					FunctionReturnCreator.Finalize();
+				FEdGraphPinType PinType;
+				K2Schema->ConvertPropertyToPinType(Property, /*out*/ PinType);
 
-					FEdGraphPinType PinType;
-					K2Schema->ConvertPropertyToPinType(Property, /*out*/ PinType);
+				UEdGraphPin* ReturnPin = ReturnNode->CreateUserDefinedPin(TEXT("ReturnValue"), PinType, EGPD_Input);
 
-					UEdGraphPin* ReturnPin = ReturnNode->CreateUserDefinedPin(TEXT("ReturnValue"), PinType, EGPD_Input);
+				// Auto-connect the pins for entry and exit, so that by default the signature is properly generated
+				UEdGraphPin* EntryNodeExec = K2Schema->FindExecutionPin(*EntryNode, EGPD_Output);
+				UEdGraphPin* ResultNodeExec = K2Schema->FindExecutionPin(*ReturnNode, EGPD_Input);
+				EntryNodeExec->MakeLinkTo(ResultNodeExec);
 
-					// Auto-connect the pins for entry and exit, so that by default the signature is properly generated
-					UEdGraphPin* EntryNodeExec = K2Schema->FindExecutionPin(*EntryNode, EGPD_Output);
-					UEdGraphPin* ResultNodeExec = K2Schema->FindExecutionPin(*ReturnNode, EGPD_Input);
-					EntryNodeExec->MakeLinkTo(ResultNodeExec);
+				FGraphNodeCreator<UK2Node_VariableGet> MemberGetCreator(*FunctionGraph);
+				UK2Node_VariableGet* VarNode = MemberGetCreator.CreateNode();
+				VarNode->VariableReference.SetSelfMember(PropertyName);
+				MemberGetCreator.Finalize();
 
-					FGraphNodeCreator<UK2Node_VariableGet> MemberGetCreator(*FunctionGraph);
-					UK2Node_VariableGet* VarNode = MemberGetCreator.CreateNode();
-					VarNode->VariableReference.SetSelfMember(PropertyName);
-					MemberGetCreator.Finalize();
+				ReturnPin->MakeLinkTo(VarNode->GetValuePin());
 
-					ReturnPin->MakeLinkTo(VarNode->GetValuePin());
+				// We need to flag the entry node to make sure that the compiled function is callable from Kismet2
+				int32 ExtraFunctionFlags = ( FUNC_Private | FUNC_Const );
+				K2Schema->AddExtraFunctionFlags(FunctionGraph, ExtraFunctionFlags);
 
-					// We need to flag the entry node to make sure that the compiled function is callable from Kismet2
-					int32 ExtraFunctionFlags = (FUNC_Private | FUNC_Const);
-					K2Schema->AddExtraFunctionFlags(FunctionGraph, ExtraFunctionFlags);
+				//Blueprint->FunctionGraphs.Add(FunctionGraph);
 
-					// new graph will be processed in Super::CreateFunctionList() call
-					Blueprint->FunctionGraphs.Add(FunctionGraph);
-				}
+				ProcessOneFunctionGraph(FunctionGraph, true);
+				//FEdGraphUtilities::MergeChildrenGraphsIn(Ubergraph, FunctionGraph, /*bRequireSchemaMatch=*/ true);
 			}
 		}
 	}
-
-	Super::CreateFunctionList();
 }
 
 void FWidgetBlueprintCompiler::ValidateWidgetNames()
diff --git a/Engine/Source/Editor/UMGEditor/Private/WidgetBlueprintCompiler.h b/Engine/Source/Editor/UMGEditor/Private/WidgetBlueprintCompiler.h
--- a/Engine/Source/Editor/UMGEditor/Private/WidgetBlueprintCompiler.h
+++ b/Engine/Source/Editor/UMGEditor/Private/WidgetBlueprintCompiler.h
@@ -30,7 +30,6 @@ protected:
 
 	// FKismetCompilerContext
 	//virtual UEdGraphSchema_K2* CreateSchema() override;
-	virtual void SetClassForBytecodeCompile(UClass* TargetClass) override;
 	virtual void CreateFunctionList() override;
 	virtual void SpawnNewClass(const FString& NewClassName) override;
 	virtual void PrecompileFunction(FKismetFunctionContext& Context) override;
diff --git a/Engine/Source/Editor/UnrealEd/Private/Kismet2/Kismet2.cpp b/Engine/Source/Editor/UnrealEd/Private/Kismet2/Kismet2.cpp
--- a/Engine/Source/Editor/UnrealEd/Private/Kismet2/Kismet2.cpp
+++ b/Engine/Source/Editor/UnrealEd/Private/Kismet2/Kismet2.cpp
@@ -914,6 +914,8 @@ void FKismetEditorUtilities::RecompileBlueprintBytecode(UBlueprint* BlueprintObj
 	TGuardValue<bool> GuardTemplateNameFlag(GCompilingBlueprint, true);
 	FCompilerResultsLog Results;
 
+	auto ReinstanceHelper = FBlueprintCompileReinstancer::Create(BlueprintObj->GeneratedClass, true);
+
 	FKismetCompilerOptions CompileOptions;
 	CompileOptions.CompileType = EKismetCompileType::BytecodeOnly;
 	{
@@ -921,6 +923,8 @@ void FKismetEditorUtilities::RecompileBlueprintBytecode(UBlueprint* BlueprintObj
 		Compiler.CompileBlueprint(BlueprintObj, CompileOptions, Results, NULL, ObjLoaded);
 	}
 
+	ReinstanceHelper->UpdateBytecodeReferences();
+
 	if (BlueprintPackage != NULL)
 	{
 		BlueprintPackage->SetDirtyFlag(bStartedWithUnsavedChanges);
diff --git a/Engine/Source/Runtime/Engine/Classes/Engine/Blueprint.h b/Engine/Source/Runtime/Engine/Classes/Engine/Blueprint.h
--- a/Engine/Source/Runtime/Engine/Classes/Engine/Blueprint.h
+++ b/Engine/Source/Runtime/Engine/Classes/Engine/Blueprint.h
@@ -359,9 +359,6 @@ class ENGINE_API UBlueprint : public UBlueprintCore
 	class USimpleConstructionScript* SimpleConstructionScript;
 
 #if WITH_EDITORONLY_DATA
-	UPROPERTY(transient, duplicatetransient)
-	UEdGraph* ConsolidatedEventGraph;
-
 	/** Set of pages that combine into a single uber-graph */
 	UPROPERTY()
 	TArray<class UEdGraph*> UbergraphPages;
diff --git a/Engine/Source/Runtime/Engine/Private/BlueprintGeneratedClass.cpp b/Engine/Source/Runtime/Engine/Private/BlueprintGeneratedClass.cpp
--- a/Engine/Source/Runtime/Engine/Private/BlueprintGeneratedClass.cpp
+++ b/Engine/Source/Runtime/Engine/Private/BlueprintGeneratedClass.cpp
@@ -184,16 +184,7 @@ struct FConditionalRecompileClassHepler
 			}
 		}
 
-		if (GeneratingBP->Status != BS_Error)
-		{
-			return ENeededAction::Recompile;
-		}
-		else
-		{
-			// don't recompile blueprints that are in the error state, doing so may clear the error state which could trigger bytecode regeneration. We 
-			// can't regenerate bytecode from a 'bad' full compile.
-			return ENeededAction::None;
-		}
+		return ENeededAction::Recompile;
 	}
 };
 
