Alleviating CIS warnings for UT - Mirroring CL 2990816 (from Dev-BP)

Modifying the fix for UE-12326 (CL 2949394) - Making it more targeted. Expanding split-pins only (instead of expanding all nodes) before ProcessAllAnimationNodes() is called.

diff --git a/Engine/Source/Editor/KismetCompiler/Private/AnimBlueprintCompiler.cpp b/Engine/Source/Editor/KismetCompiler/Private/AnimBlueprintCompiler.cpp
--- a/Engine/Source/Editor/KismetCompiler/Private/AnimBlueprintCompiler.cpp
+++ b/Engine/Source/Editor/KismetCompiler/Private/AnimBlueprintCompiler.cpp
@@ -1203,7 +1203,24 @@ void FAnimBlueprintCompiler::MergeUbergraphPagesIn(UEdGraph* Ubergraph)
 		}
 
 		// Make sure we expand any split pins here before we process animation nodes.
-		ExpansionStep(ConsolidatedEventGraph, false);
+		for (TArray<UEdGraphNode*>::TIterator NodeIt(ConsolidatedEventGraph->Nodes); NodeIt; ++NodeIt)
+		{
+			UK2Node* K2Node = Cast<UK2Node>(*NodeIt);
+			if (K2Node != nullptr)
+			{
+				// We iterate the array in reverse so we can recombine split-pins (which modifies the pins array)
+				for (int32 PinIndex = K2Node->Pins.Num() - 1; PinIndex >= 0; --PinIndex)
+				{
+					UEdGraphPin* Pin = K2Node->Pins[PinIndex];
+					if (Pin->SubPins.Num() == 0)
+					{
+						continue;
+					}
+
+					K2Node->ExpandSplitPin(this, ConsolidatedEventGraph, Pin);
+				}
+			}
+		}
 
 		// Compile the animation graph
 		ProcessAllAnimationNodes();
