Copying //UE4/Dev-Blueprints to //UE4/Dev-Main (Source: //UE4/Dev-Blueprints @ 3130440)

diff --git a/Engine/Source/Editor/Kismet/Private/Profiler/SBlueprintProfilerToolbar.cpp b/Engine/Source/Editor/Kismet/Private/Profiler/SBlueprintProfilerToolbar.cpp
--- a/Engine/Source/Editor/Kismet/Private/Profiler/SBlueprintProfilerToolbar.cpp
+++ b/Engine/Source/Editor/Kismet/Private/Profiler/SBlueprintProfilerToolbar.cpp
@@ -3,16 +3,41 @@
 #include "BlueprintEditorPrivatePCH.h"
 #include "BlueprintProfilerModule.h"
 #include "SBlueprintProfilerToolbar.h"
-#include "EventExecution.h"
 #include "SNumericEntryBox.h"
 
 #define LOCTEXT_NAMESPACE "BlueprintProfilerToolbar"
 
+//////////////////////////////////////////////////////////////////////////
+// CustomStatisticThresholdValues
+
+namespace CustomStatisticThresholdValues
+{
+	// Event thresholds
+	const float EventMin = 0.01f;
+	const float EventMax = 10.f;
+	const float EventDefault = 1.f;
+
+	// Node Average Thresholds
+	const float NodeAvgMin = 0.001f;
+	const float NodeAvgMax = 5.f;
+	const float NodeAvgDefault = 0.2f;
+
+	// Node Inclusive Thresholds
+	const float NodeInclMin = 0.001f;
+	const float NodeInclMax = 5.f;
+	const float NodeInclDefault = 0.25f;
+
+	// Node Max Thresholds
+	const float NodeMaxTimeMin = 0.01f;
+	const float NodeMaxTimeMax = 5.f;
+	const float NodeMaxTimeDefault = 0.5f;
+};
+
 //////////////////////////////////////////////////////////////////////////
 // FBPProfilerStatOptions
 
 FBlueprintProfilerStatOptions::FBlueprintProfilerStatOptions()
-	: ActiveInstance(NAME_None)
+	: ActiveInstance(SPDN_Blueprint)
 	, ActiveGraph(NAME_None)
 	, Flags(Modified)
 {
@@ -42,11 +67,6 @@ FBlueprintProfilerStatOptions::FBlueprintProfilerStatOptions()
 	{
 		Flags |= DisplayInheritedEvents;
 	}
-
-	if (ProfilerSettings->bAverageBlueprintStats)
-	{
-		Flags |= AverageBlueprintStats;
-	}
 }
 
 void FBlueprintProfilerStatOptions::SetActiveInstance(const FName InstanceName)
@@ -71,6 +91,7 @@ TSharedRef<SWidget> FBlueprintProfilerStatOptions::CreateComboContent()
 {
 	return	
 		SNew(SVerticalBox)
+#if 1 // Debugging options - not for MVP
 		+SVerticalBox::Slot()
 		.HAlign(HAlign_Left)
 		.AutoHeight()
@@ -84,13 +105,34 @@ TSharedRef<SWidget> FBlueprintProfilerStatOptions::CreateComboContent()
 				.Content()
 				[
 					SNew(STextBlock)
-					.Text(LOCTEXT("FilterToGraph_Label", "Filter to Graph"))
-					.ToolTipText(LOCTEXT("FilterToGraph_Tooltip", "Only Display events that are applicable to the current graph."))
+					.Text(LOCTEXT("ShowHottestPathStats_Label", "Show HottestPath Stats"))
+					.ToolTipText(LOCTEXT("ShowHottestPathStats_Tooltip", "Debugging feature - Show hottest path stats column."))
 				]
-				.IsChecked<FBlueprintProfilerStatOptions, uint32>(this, &FBlueprintProfilerStatOptions::GetChecked, GraphFilter)
-				.OnCheckStateChanged<FBlueprintProfilerStatOptions, uint32>(this, &FBlueprintProfilerStatOptions::OnChecked, GraphFilter)
+				.IsChecked<FBlueprintProfilerStatOptions, uint32>(this, &FBlueprintProfilerStatOptions::GetChecked, DisplayHottestPathStats)
+				.OnCheckStateChanged<FBlueprintProfilerStatOptions, uint32>(this, &FBlueprintProfilerStatOptions::OnChecked, DisplayHottestPathStats)
+			]
+		]
+		+SVerticalBox::Slot()
+		.HAlign(HAlign_Left)
+		.AutoHeight()
+		[
+			SNew(SHorizontalBox)
+			+SHorizontalBox::Slot()
+			.AutoWidth()
+			.Padding(FMargin(5,2))
+			[
+				SNew(SCheckBox)
+				.Content()
+				[
+					SNew(STextBlock)
+					.Text(LOCTEXT("ShowNodeHeatStats_Label", "Show Node Heat Level Stats"))
+					.ToolTipText(LOCTEXT("ShowNodeHeatStats_Tooltip", "Debugging feature - Show node heat level stats column."))
+				]
+				.IsChecked<FBlueprintProfilerStatOptions, uint32>(this, &FBlueprintProfilerStatOptions::GetChecked, DisplayNodeHeatStat)
+				.OnCheckStateChanged<FBlueprintProfilerStatOptions, uint32>(this, &FBlueprintProfilerStatOptions::OnChecked, DisplayNodeHeatStat)
 			]
 		]
+#endif
 		+SVerticalBox::Slot()
 		.HAlign(HAlign_Left)
 		.AutoHeight()
@@ -104,11 +146,11 @@ TSharedRef<SWidget> FBlueprintProfilerStatOptions::CreateComboContent()
 				.Content()
 				[
 					SNew(STextBlock)
-					.Text(LOCTEXT("AverageBlueprintStats_Label", "Use Averaged Blueprint Stats"))
-					.ToolTipText(LOCTEXT("AverageBlueprintStats_Tooltip", "Blueprint Statistics are Averaged or Summed By Instance."))
+					.Text(LOCTEXT("FilterToGraph_Label", "Filter to Graph"))
+					.ToolTipText(LOCTEXT("FilterToGraph_Tooltip", "Only Display events that are applicable to the current graph."))
 				]
-				.IsChecked<FBlueprintProfilerStatOptions, uint32>(this, &FBlueprintProfilerStatOptions::GetChecked, AverageBlueprintStats)
-				.OnCheckStateChanged<FBlueprintProfilerStatOptions, uint32>(this, &FBlueprintProfilerStatOptions::OnChecked, AverageBlueprintStats)
+				.IsChecked<FBlueprintProfilerStatOptions, uint32>(this, &FBlueprintProfilerStatOptions::GetChecked, GraphFilter)
+				.OnCheckStateChanged<FBlueprintProfilerStatOptions, uint32>(this, &FBlueprintProfilerStatOptions::OnChecked, GraphFilter)
 			]
 		]
 #if 0 // Removing pure timings until we have an advanced display options.
@@ -228,11 +270,6 @@ void FBlueprintProfilerStatOptions::OnChecked(ECheckBoxState NewState, const uin
 		Flags &= ~FlagsIn;
 	}
 
-	if (FlagsIn & AverageBlueprintStats)
-	{
-		FScriptPerfData::EnableBlueprintStatAverage(NewState == ECheckBoxState::Checked);
-	}
-
 	Flags |= Modified;
 
 	UBlueprintProfilerSettings* ProfilerSettings = GetMutableDefault<UBlueprintProfilerSettings>();
@@ -262,9 +299,14 @@ void FBlueprintProfilerStatOptions::OnChecked(ECheckBoxState NewState, const uin
 		ProfilerSettings->bDisplayInheritedEvents = !!(Flags & DisplayInheritedEvents);
 	}
 
-	if (FlagsIn & AverageBlueprintStats)
+	if (FlagsIn & DisplayHottestPathStats)
+	{
+		ProfilerSettings->bShowHottestPathStats = !!(Flags & DisplayHottestPathStats);
+	}
+
+	if (FlagsIn & DisplayNodeHeatStat)
 	{
-		ProfilerSettings->bAverageBlueprintStats = !!(Flags & AverageBlueprintStats);
+		ProfilerSettings->bShowHeatLevelStats = !!(Flags & DisplayNodeHeatStat);
 	}
 
 	ProfilerSettings->SaveConfig();
@@ -723,6 +765,7 @@ void SBlueprintProfilerToolbar::OnHeatLevelMetricsTypeChanged(ECheckBoxState InC
 	{
 		UBlueprintProfilerSettings* Settings = GetMutableDefault<UBlueprintProfilerSettings>();
 		Settings->HeatLevelMetricsType = NewHeatLevelMetricsType;
+		Settings->SetPerformanceThresholdsModified(true);
 
 		Settings->SaveConfig();
 	}
@@ -757,6 +800,7 @@ void SBlueprintProfilerToolbar::OnHeatMapDisplayModeChanged(const EHeatmapContro
 		{
 			Settings->GraphNodeHeatMapDisplayMode = NewHeatMapDisplayMode;
 			Settings->SaveConfig();
+			Settings->SetPerformanceThresholdsModified(true);
 		}
 	}
 	else 
@@ -765,6 +809,7 @@ void SBlueprintProfilerToolbar::OnHeatMapDisplayModeChanged(const EHeatmapContro
 		{
 			Settings->WireHeatMapDisplayMode = NewHeatMapDisplayMode;
 			Settings->SaveConfig();
+			Settings->SetPerformanceThresholdsModified(true);
 		}
 	}
 }
@@ -803,10 +848,10 @@ TSharedRef<SWidget> SBlueprintProfilerToolbar::CreateCustomHeatThresholdsMenuCon
 				[
 					SNew(SNumericEntryBox<float>)
 					.AllowSpin(true)
-					.MinValue(0.01f)
-					.MaxValue(5.f)
-					.MinSliderValue(0.01f)
-					.MaxSliderValue(5.f)
+					.MinValue(CustomStatisticThresholdValues::EventMin)
+					.MaxValue(CustomStatisticThresholdValues::EventMax)
+					.MinSliderValue(CustomStatisticThresholdValues::EventMin)
+					.MaxSliderValue(CustomStatisticThresholdValues::EventMax)
 					.Value(this, &SBlueprintProfilerToolbar::GetCustomPerformanceThreshold, ECustomPerformanceThreshold::Event)
 					.OnValueChanged(this, &SBlueprintProfilerToolbar::SetCustomPerformanceThreshold, ECustomPerformanceThreshold::Event)
 					.MinDesiredValueWidth(100)
@@ -856,10 +901,10 @@ TSharedRef<SWidget> SBlueprintProfilerToolbar::CreateCustomHeatThresholdsMenuCon
 				[
 					SNew(SNumericEntryBox<float>)
 					.AllowSpin(true)
-					.MinValue(0.001f)
-					.MaxValue(0.5f)
-					.MinSliderValue(0.001f)
-					.MaxSliderValue(0.5f)
+					.MinValue(CustomStatisticThresholdValues::NodeAvgMin)
+					.MaxValue(CustomStatisticThresholdValues::NodeAvgMax)
+					.MinSliderValue(CustomStatisticThresholdValues::NodeAvgMin)
+					.MaxSliderValue(CustomStatisticThresholdValues::NodeAvgMax)
 					.Value(this, &SBlueprintProfilerToolbar::GetCustomPerformanceThreshold, ECustomPerformanceThreshold::Average)
 					.OnValueChanged(this, &SBlueprintProfilerToolbar::SetCustomPerformanceThreshold, ECustomPerformanceThreshold::Average)
 					.MinDesiredValueWidth(100)
@@ -909,10 +954,10 @@ TSharedRef<SWidget> SBlueprintProfilerToolbar::CreateCustomHeatThresholdsMenuCon
 				[
 					SNew(SNumericEntryBox<float>)
 					.AllowSpin(true)
-					.MinValue(0.001f)
-					.MaxValue(0.5f)
-					.MinSliderValue(0.001f)
-					.MaxSliderValue(0.5f)
+					.MinValue(CustomStatisticThresholdValues::NodeInclMin)
+					.MaxValue(CustomStatisticThresholdValues::NodeInclMax)
+					.MinSliderValue(CustomStatisticThresholdValues::NodeInclMin)
+					.MaxSliderValue(CustomStatisticThresholdValues::NodeInclMax)
 					.Value(this, &SBlueprintProfilerToolbar::GetCustomPerformanceThreshold, ECustomPerformanceThreshold::Inclusive)
 					.OnValueChanged(this, &SBlueprintProfilerToolbar::SetCustomPerformanceThreshold, ECustomPerformanceThreshold::Inclusive)
 					.MinDesiredValueWidth(100)
@@ -962,10 +1007,10 @@ TSharedRef<SWidget> SBlueprintProfilerToolbar::CreateCustomHeatThresholdsMenuCon
 				[
 					SNew(SNumericEntryBox<float>)
 					.AllowSpin(true)
-					.MinValue(0.01f)
-					.MaxValue(2.f)
-					.MinSliderValue(0.01f)
-					.MaxSliderValue(2.f)
+					.MinValue(CustomStatisticThresholdValues::NodeMaxTimeMin)
+					.MaxValue(CustomStatisticThresholdValues::NodeMaxTimeMax)
+					.MinSliderValue(CustomStatisticThresholdValues::NodeMaxTimeMin)
+					.MaxSliderValue(CustomStatisticThresholdValues::NodeMaxTimeMax)
 					.Value(this, &SBlueprintProfilerToolbar::GetCustomPerformanceThreshold, ECustomPerformanceThreshold::Max)
 					.OnValueChanged(this, &SBlueprintProfilerToolbar::SetCustomPerformanceThreshold, ECustomPerformanceThreshold::Max)
 					.MinDesiredValueWidth(100)
@@ -1071,22 +1116,22 @@ EVisibility SBlueprintProfilerToolbar::IsCustomPerformanceThresholdDefaultButton
 	{
 	case ECustomPerformanceThreshold::Event:
 		CurValue = ProfilerSettings->CustomEventPerformanceThreshold;
-		DefValue = UBlueprintProfilerSettings::CustomEventPerformanceThresholdDefaultValue;
+		DefValue = CustomStatisticThresholdValues::EventDefault;
 		break;
 
 	case ECustomPerformanceThreshold::Average:
 		CurValue = ProfilerSettings->CustomAveragePerformanceThreshold;
-		DefValue = UBlueprintProfilerSettings::CustomAveragePerformanceThresholdDefaultValue;
+		DefValue = CustomStatisticThresholdValues::NodeAvgDefault;
 		break;
 
 	case ECustomPerformanceThreshold::Inclusive:
 		CurValue = ProfilerSettings->CustomInclusivePerformanceThreshold;
-		DefValue = UBlueprintProfilerSettings::CustomInclusivePerformanceThresholdDefaultValue;
+		DefValue = CustomStatisticThresholdValues::NodeInclDefault;
 		break;
 
 	case ECustomPerformanceThreshold::Max:
 		CurValue = ProfilerSettings->CustomMaxPerformanceThreshold;
-		DefValue = UBlueprintProfilerSettings::CustomMaxPerformanceThresholdDefaultValue;
+		DefValue = CustomStatisticThresholdValues::NodeMaxTimeDefault;
 		break;
 	}
 
@@ -1103,25 +1148,25 @@ FReply SBlueprintProfilerToolbar::ResetCustomPerformanceThreshold(ECustomPerform
 	{
 	case ECustomPerformanceThreshold::Event:
 		OldValue = ProfilerSettings->CustomEventPerformanceThreshold;
-		ProfilerSettings->CustomEventPerformanceThreshold = UBlueprintProfilerSettings::CustomEventPerformanceThresholdDefaultValue;
+		ProfilerSettings->CustomEventPerformanceThreshold = CustomStatisticThresholdValues::EventDefault;
 		NewValue = ProfilerSettings->CustomEventPerformanceThreshold;
 		break;
 
 	case ECustomPerformanceThreshold::Average:
 		OldValue = ProfilerSettings->CustomAveragePerformanceThreshold;
-		ProfilerSettings->CustomAveragePerformanceThreshold = UBlueprintProfilerSettings::CustomAveragePerformanceThresholdDefaultValue;
+		ProfilerSettings->CustomAveragePerformanceThreshold = CustomStatisticThresholdValues::NodeAvgDefault;
 		NewValue = ProfilerSettings->CustomAveragePerformanceThreshold;
 		break;
 
 	case ECustomPerformanceThreshold::Inclusive:
 		OldValue = ProfilerSettings->CustomInclusivePerformanceThreshold;
-		ProfilerSettings->CustomInclusivePerformanceThreshold = UBlueprintProfilerSettings::CustomInclusivePerformanceThresholdDefaultValue;
+		ProfilerSettings->CustomInclusivePerformanceThreshold = CustomStatisticThresholdValues::NodeInclDefault;
 		NewValue = ProfilerSettings->CustomInclusivePerformanceThreshold;
 		break;
 
 	case ECustomPerformanceThreshold::Max:
 		OldValue = ProfilerSettings->CustomMaxPerformanceThreshold;
-		ProfilerSettings->CustomMaxPerformanceThreshold = UBlueprintProfilerSettings::CustomMaxPerformanceThresholdDefaultValue;
+		ProfilerSettings->CustomMaxPerformanceThreshold = CustomStatisticThresholdValues::NodeMaxTimeDefault;
 		NewValue = ProfilerSettings->CustomMaxPerformanceThreshold;
 		break;
 	}
diff --git a/Engine/Source/Editor/Kismet/Private/Profiler/SBlueprintProfilerToolbar.h b/Engine/Source/Editor/Kismet/Private/Profiler/SBlueprintProfilerToolbar.h
--- a/Engine/Source/Editor/Kismet/Private/Profiler/SBlueprintProfilerToolbar.h
+++ b/Engine/Source/Editor/Kismet/Private/Profiler/SBlueprintProfilerToolbar.h
@@ -3,6 +3,7 @@
 #pragma once
 
 #include "BlueprintProfilerSettings.h"
+#include "EventExecution.h"
 
 /** Blueprint performance view type */
 namespace EBlueprintPerfViewType
@@ -45,7 +46,8 @@ public:
 		GraphFilter					= 0x00000004,	// Filter to current graph
 		DisplayPure					= 0x00000008,	// Display pure stats
 		DisplayInheritedEvents		= 0x00000010,	// Display inherited events
-		AverageBlueprintStats		= 0x00000020,	// Average blueprint stats ( Rather than sum )
+		DisplayHottestPathStats		= 0x00000020,	// -- Debugging feature to show hottest path heat levels
+		DisplayNodeHeatStat			= 0x00000040,	// -- Debugging feature to show node heat levels
 		Modified					= 0x10000000,	// Display state was changed
 	};
 
@@ -74,7 +76,7 @@ public:
 	bool HasAllFlags(const uint32 FlagsIn) const { return (Flags & FlagsIn) == FlagsIn; }
 
 	/** Set active instance name */
-	FName GetActiveInstance() const { return HasFlags(DisplayByInstance) ? ActiveInstance : NAME_None; }
+	FName GetActiveInstance() const { return HasFlags(DisplayByInstance) ? ActiveInstance : SPDN_Blueprint; }
 
 	/** Set active instance name */
 	void SetActiveInstance(const FName InstanceName);
diff --git a/Engine/Source/Editor/Kismet/Private/Profiler/SGraphExecutionStatDisplay.cpp b/Engine/Source/Editor/Kismet/Private/Profiler/SGraphExecutionStatDisplay.cpp
--- a/Engine/Source/Editor/Kismet/Private/Profiler/SGraphExecutionStatDisplay.cpp
+++ b/Engine/Source/Editor/Kismet/Private/Profiler/SGraphExecutionStatDisplay.cpp
@@ -36,7 +36,12 @@ void SGraphExecutionStatDisplay::Construct(const FArguments& InArgs)
 	{
 		Profiler->GetGraphLayoutChangedDelegate().AddSP(this, &SGraphExecutionStatDisplay::OnGraphLayoutChanged);
 	}
+	// Create the tree view
+	ConstructTreeView();
+}
 
+void SGraphExecutionStatDisplay::ConstructTreeView()
+{
 	ChildSlot
 	[
 		SNew(SOverlay)
@@ -59,46 +64,7 @@ void SGraphExecutionStatDisplay::Construct(const FArguments& InArgs)
 					.OnExpansionChanged(this, &SGraphExecutionStatDisplay::OnStatisticExpansionChanged)
 					.HeaderRow
 					(
-						SNew(SHeaderRow)
-						+SHeaderRow::Column(SProfilerStatRow::GetStatName(EBlueprintProfilerStat::Name))
-						.DefaultLabel(SProfilerStatRow::GetStatText(EBlueprintProfilerStat::Name))
-						.FillWidth(0.5f)
-						+SHeaderRow::Column(SProfilerStatRow::GetStatName(EBlueprintProfilerStat::AverageTime))
-						.DefaultLabel(SProfilerStatRow::GetStatText(EBlueprintProfilerStat::AverageTime))
-						.FixedWidth(90)
-						.HAlignHeader(HAlign_Right)
-						.HAlignCell(HAlign_Right)
-						.VAlignCell(VAlign_Center)
-						+SHeaderRow::Column(SProfilerStatRow::GetStatName(EBlueprintProfilerStat::InclusiveTime))
-						.DefaultLabel(SProfilerStatRow::GetStatText(EBlueprintProfilerStat::InclusiveTime))
-						.FixedWidth(120)
-						.HAlignHeader(HAlign_Right)
-						.HAlignCell(HAlign_Right)
-						.VAlignCell(VAlign_Center)
-						+SHeaderRow::Column(SProfilerStatRow::GetStatName(EBlueprintProfilerStat::MaxTime))
-						.DefaultLabel(SProfilerStatRow::GetStatText(EBlueprintProfilerStat::MaxTime))
-						.FixedWidth(90)
-						.HAlignHeader(HAlign_Right)
-						.HAlignCell(HAlign_Right)
-						.VAlignCell(VAlign_Center)
-						+SHeaderRow::Column(SProfilerStatRow::GetStatName(EBlueprintProfilerStat::MinTime))
-						.DefaultLabel(SProfilerStatRow::GetStatText(EBlueprintProfilerStat::MinTime))
-						.FixedWidth(90)
-						.HAlignHeader(HAlign_Right)
-						.HAlignCell(HAlign_Right)
-						.VAlignCell(VAlign_Center)
-						+SHeaderRow::Column(SProfilerStatRow::GetStatName(EBlueprintProfilerStat::TotalTime))
-						.DefaultLabel(SProfilerStatRow::GetStatText(EBlueprintProfilerStat::TotalTime))
-						.FixedWidth(100)
-						.HAlignHeader(HAlign_Right)
-						.HAlignCell(HAlign_Right)
-						.VAlignCell(VAlign_Center)
-						+SHeaderRow::Column(SProfilerStatRow::GetStatName(EBlueprintProfilerStat::Samples))
-						.DefaultLabel(SProfilerStatRow::GetStatText(EBlueprintProfilerStat::Samples))
-						.FixedWidth(100)
-						.HAlignHeader(HAlign_Center)
-						.HAlignCell(HAlign_Center)
-						.VAlignCell(VAlign_Center)
+						GenerateHeaderRow()
 					)
 				]
 			]
@@ -130,6 +96,107 @@ void SGraphExecutionStatDisplay::Construct(const FArguments& InArgs)
 	];
 }
 
+TSharedPtr<SHeaderRow> SGraphExecutionStatDisplay::GenerateHeaderRow() const
+{
+	// Create the header.
+	TSharedPtr<SHeaderRow> HeaderRow;
+	SAssignNew(HeaderRow, SHeaderRow)
+		+SHeaderRow::Column(SProfilerStatRow::GetStatName(EBlueprintProfilerStat::Name))
+		.DefaultLabel(SProfilerStatRow::GetStatText(EBlueprintProfilerStat::Name))
+		.FillWidth(0.5f);
+	// Add conditional columns
+	const UBlueprintProfilerSettings* ProfilerSettings = GetDefault<UBlueprintProfilerSettings>();
+	if (ProfilerSettings->bShowHottestPathStats)
+	{
+		HeaderRow->AddColumn(SHeaderRow::Column(SProfilerStatRow::GetStatName(EBlueprintProfilerStat::HottestPath))
+			.DefaultLabel(SProfilerStatRow::GetStatText(EBlueprintProfilerStat::HottestPath))
+			.FixedWidth(90)
+			.HAlignHeader(HAlign_Right)
+			.HAlignCell(HAlign_Right)
+			.VAlignCell(VAlign_Center));
+	}	
+	if (ProfilerSettings->bShowHeatLevelStats)
+	{
+		HeaderRow->AddColumn(SHeaderRow::Column(SProfilerStatRow::GetStatName(EBlueprintProfilerStat::HeatLevel))
+			.DefaultLabel(SProfilerStatRow::GetStatText(EBlueprintProfilerStat::HeatLevel))
+			.FixedWidth(90)
+			.HAlignHeader(HAlign_Right)
+			.HAlignCell(HAlign_Right)
+			.VAlignCell(VAlign_Center));
+	}
+	// Add Standard columns
+	HeaderRow->AddColumn(SHeaderRow::Column(SProfilerStatRow::GetStatName(EBlueprintProfilerStat::AverageTime))
+		.DefaultLabel(SProfilerStatRow::GetStatText(EBlueprintProfilerStat::AverageTime))
+		.FixedWidth(90)
+		.HAlignHeader(HAlign_Right)
+		.HAlignCell(HAlign_Right)
+		.VAlignCell(VAlign_Center)
+		.OnSort(this, &SGraphExecutionStatDisplay::SetAverageHeatDisplay));
+	HeaderRow->AddColumn(SHeaderRow::Column(SProfilerStatRow::GetStatName(EBlueprintProfilerStat::InclusiveTime))
+		.DefaultLabel(SProfilerStatRow::GetStatText(EBlueprintProfilerStat::InclusiveTime))
+		.FixedWidth(120)
+		.HAlignHeader(HAlign_Right)
+		.HAlignCell(HAlign_Right)
+		.VAlignCell(VAlign_Center)
+		.OnSort(this, &SGraphExecutionStatDisplay::SetInclusiveHeatDisplay));
+	HeaderRow->AddColumn(SHeaderRow::Column(SProfilerStatRow::GetStatName(EBlueprintProfilerStat::MaxTime))
+		.DefaultLabel(SProfilerStatRow::GetStatText(EBlueprintProfilerStat::MaxTime))
+		.FixedWidth(90)
+		.HAlignHeader(HAlign_Right)
+		.HAlignCell(HAlign_Right)
+		.VAlignCell(VAlign_Center)
+		.OnSort(this, &SGraphExecutionStatDisplay::SetMaxHeatDisplay));
+	HeaderRow->AddColumn(SHeaderRow::Column(SProfilerStatRow::GetStatName(EBlueprintProfilerStat::MinTime))
+		.DefaultLabel(SProfilerStatRow::GetStatText(EBlueprintProfilerStat::MinTime))
+		.FixedWidth(90)
+		.HAlignHeader(HAlign_Right)
+		.HAlignCell(HAlign_Right)
+		.VAlignCell(VAlign_Center));
+	HeaderRow->AddColumn(SHeaderRow::Column(SProfilerStatRow::GetStatName(EBlueprintProfilerStat::TotalTime))
+		.DefaultLabel(SProfilerStatRow::GetStatText(EBlueprintProfilerStat::TotalTime))
+		.FixedWidth(100)
+		.HAlignHeader(HAlign_Right)
+		.HAlignCell(HAlign_Right)
+		.VAlignCell(VAlign_Center)
+		.OnSort(this, &SGraphExecutionStatDisplay::SetTotalTimeHeatDisplay));
+	HeaderRow->AddColumn(SHeaderRow::Column(SProfilerStatRow::GetStatName(EBlueprintProfilerStat::Samples))
+		.DefaultLabel(SProfilerStatRow::GetStatText(EBlueprintProfilerStat::Samples))
+		.FixedWidth(100)
+		.HAlignHeader(HAlign_Center)
+		.HAlignCell(HAlign_Center)
+		.VAlignCell(VAlign_Center));
+	// Return the header row.
+	return HeaderRow;
+}
+
+void SGraphExecutionStatDisplay::SetAverageHeatDisplay(EColumnSortPriority::Type /*NotUsed*/,const FName& /*NotUsed*/, EColumnSortMode::Type /*NotUsed*/) const
+{
+	UBlueprintProfilerSettings* ProfilerSettings = GetMutableDefault<UBlueprintProfilerSettings>();
+	const bool bDisable = ProfilerSettings->GraphNodeHeatMapDisplayMode == EBlueprintProfilerHeatMapDisplayMode::Average;
+	ProfilerSettings->GraphNodeHeatMapDisplayMode = bDisable ? EBlueprintProfilerHeatMapDisplayMode::None : EBlueprintProfilerHeatMapDisplayMode::Average;
+}
+
+void SGraphExecutionStatDisplay::SetInclusiveHeatDisplay(EColumnSortPriority::Type /*NotUsed*/,const FName& /*NotUsed*/, EColumnSortMode::Type /*NotUsed*/) const
+{
+	UBlueprintProfilerSettings* ProfilerSettings = GetMutableDefault<UBlueprintProfilerSettings>();
+	const bool bDisable = ProfilerSettings->GraphNodeHeatMapDisplayMode == EBlueprintProfilerHeatMapDisplayMode::Inclusive;
+	ProfilerSettings->GraphNodeHeatMapDisplayMode = bDisable ? EBlueprintProfilerHeatMapDisplayMode::None : EBlueprintProfilerHeatMapDisplayMode::Inclusive;
+}
+
+void SGraphExecutionStatDisplay::SetMaxHeatDisplay(EColumnSortPriority::Type /*NotUsed*/,const FName& /*NotUsed*/, EColumnSortMode::Type /*NotUsed*/) const
+{
+	UBlueprintProfilerSettings* ProfilerSettings = GetMutableDefault<UBlueprintProfilerSettings>();
+	const bool bDisable = ProfilerSettings->GraphNodeHeatMapDisplayMode == EBlueprintProfilerHeatMapDisplayMode::MaxTiming;
+	ProfilerSettings->GraphNodeHeatMapDisplayMode = bDisable ? EBlueprintProfilerHeatMapDisplayMode::None : EBlueprintProfilerHeatMapDisplayMode::MaxTiming;
+}
+
+void SGraphExecutionStatDisplay::SetTotalTimeHeatDisplay(EColumnSortPriority::Type /*NotUsed*/,const FName& /*NotUsed*/, EColumnSortMode::Type /*NotUsed*/) const
+{
+	UBlueprintProfilerSettings* ProfilerSettings = GetMutableDefault<UBlueprintProfilerSettings>();
+	const bool bDisable = ProfilerSettings->GraphNodeHeatMapDisplayMode == EBlueprintProfilerHeatMapDisplayMode::Total;
+	ProfilerSettings->GraphNodeHeatMapDisplayMode = bDisable ? EBlueprintProfilerHeatMapDisplayMode::None : EBlueprintProfilerHeatMapDisplayMode::Total;
+}
+
 void SGraphExecutionStatDisplay::OnGraphLayoutChanged(TWeakObjectPtr<UBlueprint> Blueprint)
 {
 	if (CurrentBlueprint == Blueprint)
@@ -190,7 +257,7 @@ void SGraphExecutionStatDisplay::Tick(const FGeometry& AllottedGeometry, const d
 					UBlueprint* NewBlueprint = Cast<UBlueprint>((*Blueprints)[0]);
 					UBlueprintGeneratedClass* BPGC = Cast<UBlueprintGeneratedClass>(NewBlueprint->GeneratedClass);
 					TWeakObjectPtr<const UObject> NewInstance = NewBlueprint->GetObjectBeingDebugged();
-					TSharedPtr<FBlueprintExecutionContext> BlueprintExecContext = Profiler->GetBlueprintContext(BPGC->GetPathName());
+					TSharedPtr<FBlueprintExecutionContext> BlueprintExecContext = Profiler->FindBlueprintContext(BPGC->GetPathName());
 
 					if (BlueprintExecContext.IsValid())
 					{
@@ -213,6 +280,7 @@ void SGraphExecutionStatDisplay::Tick(const FGeometry& AllottedGeometry, const d
 						}
 						if (DisplayOptions->IsStateModified())
 						{
+							ConstructTreeView();
 							TSharedPtr<FScriptExecutionBlueprint> BlueprintExecNode = BlueprintExecContext->GetBlueprintExecNode();
 							RootTreeItems.Reset(0);
 
@@ -227,7 +295,7 @@ void SGraphExecutionStatDisplay::Tick(const FGeometry& AllottedGeometry, const d
 								{
 									if (DisplayOptions->HasFlags(FBlueprintProfilerStatOptions::ScopeToDebugInstance))
 									{
-										if (CurrentInstancePath != NAME_None)
+										if (CurrentInstancePath != SPDN_Blueprint)
 										{
 											TSharedPtr<FScriptExecutionNode> InstanceStat = BlueprintExecNode->GetInstanceByName(CurrentInstancePath);
 											if (InstanceStat.IsValid())
diff --git a/Engine/Source/Editor/Kismet/Private/Profiler/SGraphExecutionStatDisplay.h b/Engine/Source/Editor/Kismet/Private/Profiler/SGraphExecutionStatDisplay.h
--- a/Engine/Source/Editor/Kismet/Private/Profiler/SGraphExecutionStatDisplay.h
+++ b/Engine/Source/Editor/Kismet/Private/Profiler/SGraphExecutionStatDisplay.h
@@ -33,6 +33,24 @@ public:
 
 protected:
 
+	/** Called to construct the stat tree view */
+	void ConstructTreeView();
+
+	/** Called to generate the execution stat tree header row */
+	TSharedPtr<SHeaderRow> GenerateHeaderRow() const;
+
+	/** Callback to change the heat display mode on column sort */
+	void SetAverageHeatDisplay(EColumnSortPriority::Type /*NotUsed*/,const FName& /*NotUsed*/, EColumnSortMode::Type /*NotUsed*/) const;
+
+	/** Callback to change the heat display mode on column sort */
+	void SetInclusiveHeatDisplay(EColumnSortPriority::Type /*NotUsed*/,const FName& /*NotUsed*/, EColumnSortMode::Type /*NotUsed*/) const;
+
+	/** Callback to change the heat display mode on column sort */
+	void SetMaxHeatDisplay(EColumnSortPriority::Type /*NotUsed*/,const FName& /*NotUsed*/, EColumnSortMode::Type /*NotUsed*/) const;
+
+	/** Callback to change the heat display mode on column sort */
+	void SetTotalTimeHeatDisplay(EColumnSortPriority::Type /*NotUsed*/,const FName& /*NotUsed*/, EColumnSortMode::Type /*NotUsed*/) const;
+
 	/** Called when blueprint graph layout changes occur */
 	void OnGraphLayoutChanged(TWeakObjectPtr<UBlueprint> Blueprint);
 
diff --git a/Engine/Source/Editor/Kismet/Private/Profiler/ScriptPerfData.cpp b/Engine/Source/Editor/Kismet/Private/Profiler/ScriptPerfData.cpp
--- a/Engine/Source/Editor/Kismet/Private/Profiler/ScriptPerfData.cpp
+++ b/Engine/Source/Editor/Kismet/Private/Profiler/ScriptPerfData.cpp
@@ -3,6 +3,15 @@
 #include "BlueprintEditorPrivatePCH.h"
 #include "ScriptPerfData.h"
 
+//////////////////////////////////////////////////////////////////////////
+// ScriptPerfDataUIText
+
+namespace ScriptPerfDataUIText
+{
+	const FText StatValue_NoData(NSLOCTEXT("ScriptProfilerDataUI", "StatValue_NoData", ""));
+	const FText StatValue_NotApplicable(NSLOCTEXT("ScriptProfilerDataUI", "StatValue_NotApplicable", "--"));
+};
+
 //////////////////////////////////////////////////////////////////////////
 // FScriptHeatLevelMetrics
 
@@ -14,11 +23,9 @@ float FScriptHeatLevelMetrics::NodeTotalTimeWaterMark = 0.f;
 FNumberFormattingOptions FScriptPerfData::StatNumberFormat;
 FNumberFormattingOptions FScriptPerfData::TimeNumberFormat;
 
-bool FScriptPerfData::bAverageBlueprintStats = true;
-
 void FScriptPerfData::AddEventTiming(const double AverageTimingIn)
 {
-	const double TimingInMs = AverageTimingIn * 1000;
+	const double TimingInMs = AverageTimingIn * 1000.0;
 	AverageTiming += TimingInMs;
 	RawSamples++;
 	const int32 SampleCount = GetSampleCount();
@@ -33,7 +40,7 @@ void FScriptPerfData::AddEventTiming(const double AverageTimingIn)
 void FScriptPerfData::AddInclusiveTiming(const double InclusiveNodeTimingIn, const bool bIncrementSamples)
 {
 	// Discrete timings should negate the discrete nodetimings so they sum correctly later.
-	InclusiveTiming += InclusiveNodeTimingIn * 1000;
+	InclusiveTiming += InclusiveNodeTimingIn * 1000.0;
 	if (bIncrementSamples)
 	{
 		RawSamples++;
@@ -53,52 +60,71 @@ void FScriptPerfData::AddData(const FScriptPerfData& DataIn)
 	}
 }
 
-void FScriptPerfData::InitialiseFromDataSet(const TArray<TSharedPtr<FScriptPerfData>>& DataSet)
+void FScriptPerfData::CreateBlueprintStats(const TArray<TSharedPtr<FScriptPerfData>>& BlueprintDataSet)
 {
 	Reset();
-	const double Denom = 1.0 / DataSet.Num();
-	for (auto DataPoint : DataSet)
+	SampleFrequency = 0.f;
+	double NewMaxTiming = 0.0;
+	double NewMinTiming = 0.0;
+	for (auto DataPoint : BlueprintDataSet)
 	{
 		if (DataPoint.IsValid())
 		{
-			AverageTiming += DataPoint->AverageTiming * Denom;
-			InclusiveTiming += DataPoint->InclusiveTiming * Denom;
+			const float SampleCount = DataPoint->GetSampleCount();
+			const double AverageTimeTemp = DataPoint->GetAverageTiming();
+			AverageTiming += AverageTimeTemp * SampleCount;
+			InclusiveTiming += (DataPoint->GetInclusiveTiming() - AverageTimeTemp) * SampleCount;
+			NewMaxTiming += DataPoint->GetMaxTiming();
+			NewMinTiming += DataPoint->GetMinTiming();
 			// Either sum the samples to get the average or use the highest sample rate encountered in the dataset.
-			RawSamples = bAverageBlueprintStats ? (RawSamples + DataPoint->RawSamples) : FMath::Max(RawSamples, DataPoint->RawSamples);
+			RawSamples += SampleCount;
+			SampleFrequency += 1.f;
 		}
 	}
-	if (IsDataValid())
+	// Update min/max timings if the data is valid.
+	if (NewMaxTiming > 0.0)
 	{
-		const double NormalisedAverageTiming = GetAverageTiming();
-		MaxTiming = FMath::Max<double>(MaxTiming, NormalisedAverageTiming);
-		MinTiming = FMath::Min<double>(MinTiming, NormalisedAverageTiming);
+		MaxTiming = NewMaxTiming;
+		MinTiming = NewMinTiming;
 	}
+	SampleFrequency = FMath::Max(1.f, SampleFrequency);
 }
 
-void FScriptPerfData::AccumulateDataSet(const TArray<TSharedPtr<FScriptPerfData>>& DataSet)
+void FScriptPerfData::AccumulateDataSet(const TArray<TSharedPtr<FScriptPerfData>>& DataSet, bool bApplyAsInclusive)
 {
 	Reset();
 	SampleFrequency = 0.f;
+	double NewMaxTiming = 0.0;
+	double NewMinTiming = 0.0;
 	for (auto DataPoint : DataSet)
 	{
 		if (DataPoint.IsValid())
 		{
 			const float SampleCount = DataPoint->GetSampleCount();
-			const double AverageTimeTemp = DataPoint->GetAverageTiming();
-			AverageTiming += AverageTimeTemp * SampleCount;
-			InclusiveTiming += (DataPoint->GetInclusiveTiming() - AverageTimeTemp) * SampleCount;
+			if (bApplyAsInclusive)
+			{
+				InclusiveTiming += DataPoint->GetInclusiveTiming() * SampleCount;
+			}
+			else
+			{
+				const double AverageTimeTemp = DataPoint->GetAverageTiming();
+				AverageTiming += AverageTimeTemp * SampleCount;
+				InclusiveTiming += (DataPoint->GetInclusiveTiming() - AverageTimeTemp) * SampleCount;
+			}
+			NewMaxTiming += DataPoint->GetMaxTiming();
+			NewMinTiming += DataPoint->GetMinTiming();
 			// Either sum the samples to get the average or use the highest sample rate encountered in the dataset.
 			RawSamples += SampleCount;
 			SampleFrequency += 1.f;
 		}
 	}
-	SampleFrequency = FMath::Max(1.f, SampleFrequency);
-	if (IsDataValid())
+	// Update min/max timings if the data is valid.
+	if (NewMaxTiming > 0.0)
 	{
-		const double NormalisedAverageTiming = GetAverageTiming();
-		MaxTiming = FMath::Max<double>(MaxTiming, NormalisedAverageTiming);
-		MinTiming = FMath::Min<double>(MinTiming, NormalisedAverageTiming);
+		MaxTiming = NewMaxTiming;
+		MinTiming = NewMinTiming;
 	}
+	SampleFrequency = FMath::Max(1.f, SampleFrequency);
 }
 
 void FScriptPerfData::Reset()
@@ -136,29 +162,30 @@ double FScriptPerfData::GetInclusiveTiming() const
 	return GetAverageTiming() + InclusiveTime;
 }
 
-void FScriptPerfData::EnableBlueprintStatAverage(const bool bEnable)
-{
-	bAverageBlueprintStats = bEnable;
-}
-
 FText FScriptPerfData::GetAverageTimingText() const
 {
-	if (GetSampleCount() > 0)
+	if (AverageTiming != 0.f)
 	{
 		const double Value = GetAverageTiming();
 		return FText::AsNumber(Value, &StatNumberFormat);
 	}
-	return FText::GetEmpty();
+	else
+	{
+		return GetSampleCount() > 0.f ? ScriptPerfDataUIText::StatValue_NotApplicable : ScriptPerfDataUIText::StatValue_NoData;
+	}
 }
 
 FText FScriptPerfData::GetInclusiveTimingText() const
 {
-	if (GetSampleCount() > 0)
+	const double Value = GetInclusiveTiming();
+	if (Value != 0.f)
+	{
+		return FText::AsNumber(Value, &StatNumberFormat);
+	}
+	else
 	{
-		const double InclusiveTemp = GetInclusiveTiming();
-		return FText::AsNumber(InclusiveTemp, &StatNumberFormat);
+		return GetSampleCount() > 0.f ? ScriptPerfDataUIText::StatValue_NotApplicable : ScriptPerfDataUIText::StatValue_NoData;
 	}
-	return FText::GetEmpty();
 }
 
 FText FScriptPerfData::GetMaxTimingText() const
@@ -167,7 +194,10 @@ FText FScriptPerfData::GetMaxTimingText() const
 	{
 		return FText::AsNumber(MaxTiming, &StatNumberFormat);
 	}
-	return FText::GetEmpty();
+	else
+	{
+		return GetSampleCount() > 0.f ? ScriptPerfDataUIText::StatValue_NotApplicable : ScriptPerfDataUIText::StatValue_NoData;
+	}
 }
 
 FText FScriptPerfData::GetMinTimingText() const
@@ -176,16 +206,19 @@ FText FScriptPerfData::GetMinTimingText() const
 	{
 		return FText::AsNumber(MinTiming, &StatNumberFormat);
 	}
-	return FText::GetEmpty();
+	else
+	{
+		return GetSampleCount() > 0.f ? ScriptPerfDataUIText::StatValue_NotApplicable : ScriptPerfDataUIText::StatValue_NoData;
+	}
 }
 
 FText FScriptPerfData::GetTotalTimingText() const
 {
-	if (GetSampleCount() > 0)
+	if (GetSampleCount() > 0.f)
 	{
 		return FText::AsNumber(AverageTiming, &TimeNumberFormat);
 	}
-	return FText::GetEmpty();
+	return ScriptPerfDataUIText::StatValue_NoData;
 }
 
 FText FScriptPerfData::GetSamplesText() const
@@ -195,17 +228,28 @@ FText FScriptPerfData::GetSamplesText() const
 	{
 		return FText::AsNumber(SampleCount);
 	}
-	return FText::GetEmpty();
+	return ScriptPerfDataUIText::StatValue_NoData;
+}
+
+FText FScriptPerfData::GetHottestPathText() const
+{
+	return FText::AsNumber(HottestPathHeatValue);
+}
+
+FText FScriptPerfData::GetHeatLevelText() const
+{
+	return FText::AsNumber(AverageTimingHeatLevel);
 }
 
-void FScriptPerfData::SetHeatLevels(TSharedPtr<FScriptHeatLevelMetrics> HeatLevelMetrics)
+void FScriptPerfData::SetHeatLevels(TSharedPtr<const FScriptHeatLevelMetrics> HeatLevelMetrics)
 {
 	if (ensure(HeatLevelMetrics.IsValid()))
 	{
 		// Note: We're intentionally not accessing AverageTiming/InclusiveTiming members directly here - the accessors make some internal adjustments.
-		AverageTimingHeatLevel = FMath::Min<float>(GetAverageTiming() * HeatLevelMetrics->AveragePerformanceThreshold, 1.f);
+		const float AverageThreshold = (StatType == EScriptPerfDataType::Event) ? HeatLevelMetrics->EventPerformanceThreshold : HeatLevelMetrics->AveragePerformanceThreshold;
+		AverageTimingHeatLevel = FMath::Min<float>(GetAverageTiming() * AverageThreshold, 1.f);
 		InclusiveTimingHeatLevel = FMath::Min<float>(GetInclusiveTiming() * HeatLevelMetrics->InclusivePerformanceThreshold, 1.f);
-		MaxTimingHeatLevel = FMath::Min<float>(MaxTiming * HeatLevelMetrics->MaxPerformanceThreshold, 1.f);
+		MaxTimingHeatLevel = MaxTiming == -MAX_dbl ? 0.f : FMath::Min<float>(MaxTiming * HeatLevelMetrics->MaxPerformanceThreshold, 1.f);
 
 		if (HeatLevelMetrics->bUseTotalTimeWaterMark)
 		{
diff --git a/Engine/Source/Editor/Kismet/Private/SBlueprintPalette.cpp b/Engine/Source/Editor/Kismet/Private/SBlueprintPalette.cpp
--- a/Engine/Source/Editor/Kismet/Private/SBlueprintPalette.cpp
+++ b/Engine/Source/Editor/Kismet/Private/SBlueprintPalette.cpp
@@ -290,7 +290,7 @@ static void GetSubGraphIcon(FEdGraphSchemaAction_K2Graph const* const ActionIn,
  * @param  ColorOut		An output color, further denoting the specified action.
  * @param  ToolTipOut	An output tooltip, best describing the specified action type.
  */
-static void GetPaletteItemIcon(TSharedPtr<FEdGraphSchemaAction> ActionIn, UBlueprint const* BlueprintIn, FSlateBrush const*& BrushOut, FSlateColor& ColorOut, FText& ToolTipOut, FString& DocLinkOut, FString& DocExcerptOut)
+static void GetPaletteItemIcon(TSharedPtr<FEdGraphSchemaAction> ActionIn, UBlueprint const* BlueprintIn, FSlateBrush const*& BrushOut, FSlateColor& ColorOut, FText& ToolTipOut, FString& DocLinkOut, FString& DocExcerptOut, FSlateBrush const*& SecondaryBrushOut, FSlateColor& SecondaryColorOut)
 {
 	// Default to tooltip based on action supplied
 	ToolTipOut = (ActionIn->GetTooltipDescription().Len() > 0) ? FText::FromString(ActionIn->GetTooltipDescription()) : ActionIn->GetMenuDescription();
@@ -344,7 +344,7 @@ static void GetPaletteItemIcon(TSharedPtr<FEdGraphSchemaAction> ActionIn, UBluep
 		FEdGraphSchemaAction_K2Var* VarAction = (FEdGraphSchemaAction_K2Var*)ActionIn.Get();
 
 		UClass* VarClass = VarAction->GetVariableClass();
-		BrushOut = FBlueprintEditor::GetVarIconAndColor(VarClass, VarAction->GetVariableName(), ColorOut);
+		BrushOut = FBlueprintEditor::GetVarIconAndColor(VarClass, VarAction->GetVariableName(), ColorOut, SecondaryBrushOut, SecondaryColorOut);
 		ToolTipOut = FText::FromString(GetVarType(VarClass, VarAction->GetVariableName(), true, true));
 
 		DocLinkOut = TEXT("Shared/Editor/Blueprint/VariableTypes");
@@ -355,7 +355,7 @@ static void GetPaletteItemIcon(TSharedPtr<FEdGraphSchemaAction> ActionIn, UBluep
 		FEdGraphSchemaAction_K2LocalVar* LocalVarAction = (FEdGraphSchemaAction_K2LocalVar*)ActionIn.Get();
 
 		UStruct* VarScope = LocalVarAction->GetVariableScope();
-		BrushOut = FBlueprintEditor::GetVarIconAndColor(VarScope, LocalVarAction->GetVariableName(), ColorOut);
+		BrushOut = FBlueprintEditor::GetVarIconAndColor(VarScope, LocalVarAction->GetVariableName(), ColorOut, SecondaryBrushOut, SecondaryColorOut);
 		ToolTipOut = FText::FromString(GetVarType(VarScope, LocalVarAction->GetVariableName(), true));
 
 		DocLinkOut = TEXT("Shared/Editor/Blueprint/VariableTypes");
@@ -994,11 +994,13 @@ void SBlueprintPaletteItem::Construct(const FArguments& InArgs, FCreateWidgetFor
 	
 	// construct the icon widget
 	FSlateBrush const* IconBrush   = FEditorStyle::GetBrush(TEXT("NoBrush"));
+	FSlateBrush const* SecondaryBrush = FEditorStyle::GetBrush(TEXT("NoBrush"));
 	FSlateColor        IconColor   = FSlateColor::UseForeground();
+	FSlateColor        SecondaryIconColor   = FSlateColor::UseForeground();
 	FText			   IconToolTip = FText::FromString(GraphAction->GetTooltipDescription());
 	FString			   IconDocLink, IconDocExcerpt;
-	GetPaletteItemIcon(GraphAction, Blueprint, IconBrush, IconColor, IconToolTip, IconDocLink, IconDocExcerpt);
-	TSharedRef<SWidget> IconWidget = CreateIconWidget(IconToolTip, IconBrush, IconColor, IconDocLink, IconDocExcerpt);
+	GetPaletteItemIcon(GraphAction, Blueprint, IconBrush, IconColor, IconToolTip, IconDocLink, IconDocExcerpt, SecondaryBrush, SecondaryIconColor);
+	TSharedRef<SWidget> IconWidget = CreateIconWidget(IconToolTip, IconBrush, IconColor, IconDocLink, IconDocExcerpt, SecondaryBrush, SecondaryIconColor);
 	IconWidget->SetEnabled(!bIsFullyReadOnly);
 
 	// Setup a meta tag for this node
diff --git a/Engine/Source/Editor/Kismet/Public/BlueprintEditor.h b/Engine/Source/Editor/Kismet/Public/BlueprintEditor.h
--- a/Engine/Source/Editor/Kismet/Public/BlueprintEditor.h
+++ b/Engine/Source/Editor/Kismet/Public/BlueprintEditor.h
@@ -334,7 +334,7 @@ public:
 	 * @param IconColorOut		The resulting color for the glyph
 	 * @return					The resulting glyph brush
 	 */
-	static FSlateBrush const* GetVarIconAndColor(const UStruct* VarScope, FName VarName, FSlateColor& IconColorOut);
+	static FSlateBrush const* GetVarIconAndColor(const UStruct* VarScope, FName VarName, FSlateColor& IconColorOut, FSlateBrush const*& SecondaryBrushOut, FSlateColor& SecondaryColorOut);
 
 	/** Overridable function for determining if the current mode can script */
 	virtual bool IsInAScriptingMode() const;
diff --git a/Engine/Source/Editor/Kismet/Public/Profiler/BlueprintProfilerSettings.h b/Engine/Source/Editor/Kismet/Public/Profiler/BlueprintProfilerSettings.h
--- a/Engine/Source/Editor/Kismet/Public/Profiler/BlueprintProfilerSettings.h
+++ b/Engine/Source/Editor/Kismet/Public/Profiler/BlueprintProfilerSettings.h
@@ -53,10 +53,6 @@ public:
 	UPROPERTY(config)
 	bool bDisplayInheritedEvents;
 
-	/** Blueprint Statistics are Averaged or Summed By Instance */
-	UPROPERTY(config)
-	bool bAverageBlueprintStats;
-
 	/** Current graph node heat map display mode */
 	UPROPERTY(config)
 	EBlueprintProfilerHeatMapDisplayMode GraphNodeHeatMapDisplayMode;
@@ -85,12 +81,11 @@ public:
 	UPROPERTY(config)
 	float CustomMaxPerformanceThreshold;
 
-	/** Default custom threshold values */
-	/** @TODO - maybe we don't have a need for these anymore? they're currently supporting the "reset to default" feature in the UI - not sure we have that in the "normal" editor settings UI */
-	static float CustomEventPerformanceThresholdDefaultValue;
-	static float CustomAveragePerformanceThresholdDefaultValue;
-	static float CustomInclusivePerformanceThresholdDefaultValue;
-	static float CustomMaxPerformanceThresholdDefaultValue;
+	/** Show debugging data for hottest path stats */
+	bool bShowHottestPathStats;
+
+	/** Show debugging data for node heat stats */
+	bool bShowHeatLevelStats;
 
 public:
 
diff --git a/Engine/Source/Editor/Kismet/Public/Profiler/EventExecution.h b/Engine/Source/Editor/Kismet/Public/Profiler/EventExecution.h
--- a/Engine/Source/Editor/Kismet/Public/Profiler/EventExecution.h
+++ b/Engine/Source/Editor/Kismet/Public/Profiler/EventExecution.h
@@ -34,6 +34,8 @@ namespace EScriptExecutionNodeFlags
 		PureChain					= 0x00100000,	// Pure node call chain
 		CyclicLinkage				= 0x00200000,	// Marks execution path as cyclic.
 		InvalidTrace				= 0x00400000,	// Indicates that node doesn't contain a valid script trace.
+		RuntimeEvent				= 0x01000000,	// Event that is considered part of the runtime cost.
+		ConstructionEvent			= 0x02000000,	// Event that is considered part of the construction cost.
 		// Groups
 		CallSite					= FunctionCall|ParentFunctionCall|MacroCall|TunnelInstance,
 		BranchNode					= ConditionalBranch|SequentialBranch,
@@ -47,6 +49,12 @@ namespace EScriptExecutionNodeFlags
 	};
 }
 
+//////////////////////////////////////////////////////////////////////////
+// Script Perf Data Name Constants
+
+static const FName SPDN_Blueprint = FName(TEXT("__SPDN_Blueprint"));;
+static const FName SPDN_Node = FName(TEXT("__SPDN_Node"));
+
 //////////////////////////////////////////////////////////////////////////
 // Stat Container Type
 
@@ -124,6 +132,7 @@ public:
 
 	FScriptNodePerfData(const int32 SampleFreqIn = 1)
 		: SampleFrequency(SampleFreqIn)
+		, NodePerfData(EScriptPerfDataType::Node)
 	{
 	}
 
@@ -148,14 +157,20 @@ public:
 	/** Get global blueprint perf data for the trace path */
 	TSharedPtr<FScriptPerfData> GetBlueprintPerfDataByTracePath(const FTracePath& TracePath);
 
-	/** Get global blueprint perf data for all trace paths */
-	virtual void GetBlueprintPerfDataForAllTracePaths(FScriptPerfData& OutPerfData);
+	/** Get node performance data */
+	const FScriptPerfData& GetNodePerfData() const { return NodePerfData; }
+
+	/** Update all perf data for node */
+	virtual void UpdatePerfDataForNode();
 
 	/** Set the sample frequency */
 	void SetSampleFrequency(const int32 NewSampleFrequency) { SampleFrequency = NewSampleFrequency; }
 
 protected:
 
+	/** Is Valid Instance Name */
+	bool IsValidInstanceName(const FName InstanceName) const { return InstanceName != SPDN_Blueprint && InstanceName != SPDN_Node; }
+
 	/** Returns performance data type */
 	virtual const EScriptPerfDataType GetPerfDataType() const { return EScriptPerfDataType::Node; }
 
@@ -163,6 +178,8 @@ protected:
 	
 	/** Expected sample frequency for this perf data container */
 	int32 SampleFrequency;
+	/** Node performance data */
+	FScriptPerfData NodePerfData;
 	/** FScriptExeutionPath hash to perf data */
 	TMap<FName, TMap<const uint32, TSharedPtr<FScriptPerfData>>> InstanceInputPinToPerfDataMap;
 
@@ -200,21 +217,46 @@ struct KISMET_API FScriptExecNodeParams
 
 struct FScriptExecutionHottestPathParams
 {
-	FScriptExecutionHottestPathParams(const FName InstanceNameIn, const float EventTimingIn)
+public:
+
+	FScriptExecutionHottestPathParams(const FName InstanceNameIn, const double EventTimingIn, const int32 EventSampleCountIn, TSharedPtr<FScriptHeatLevelMetrics> HeatLevelThresholdsIn)
 		: InstanceName(InstanceNameIn)
 		, EventTiming(EventTimingIn)
-		, TimeSoFar(0.0)
+		, EventSampleCount(EventSampleCountIn)
+		, PathTimeSoFar(0.0)
+		, HeatLevelThresholds(HeatLevelThresholdsIn)
 	{
 	}
 
+	/** Get instance name */
+	const FName GetInstanceName() const { return InstanceName; }
+
+	/** Get the current tracepath */
+	FTracePath& GetTracePath() { return TracePath; }
+
+	/** Get heat thresholds */
+	TSharedPtr<const FScriptHeatLevelMetrics> GetHeatThresholds() const { return HeatLevelThresholds; }
+
+	/** Calculate heat level */
+	const float CalculateHeatLevel() const { return EventTiming > 0.0 ? static_cast<float>((EventTiming - PathTimeSoFar) / EventTiming) : 0.f; }
+
+	/** Add path timings, adjusting for variation in stat execution frequencies */
+	void AddPathTiming(const double PathTiming, const int32 SampleCount) { PathTimeSoFar += PathTiming * (SampleCount / EventSampleCount); }
+
+private:
+
 	/** Instance name */
 	const FName InstanceName;
 	/** Event total timing */
-	const double EventTiming;
-	/** Event time so far */
-	double TimeSoFar;
+	double EventTiming;
+	/** Top level event samples */
+	double EventSampleCount;
+	/** Event path time so far */
+	double PathTimeSoFar;
 	/** Current tracepath */
 	FTracePath TracePath;
+	/** Heat level thresholds */
+	TSharedPtr<const FScriptHeatLevelMetrics> HeatLevelThresholds;
 };
 
 //////////////////////////////////////////////////////////////////////////
@@ -263,6 +305,9 @@ public:
 	/** Returns if this exec event represents the start of a delegate pin event execution path */
 	bool IsEventPin() const { return (NodeFlags & EScriptExecutionNodeFlags::EventPin) != 0U; }
 
+	/** Returns if this exec event happens at runtime ( alternative is construction time ) */
+	bool IsRuntimeEvent() const { return (NodeFlags & EScriptExecutionNodeFlags::RuntimeEvent) != 0U; }
+
 	/** Returns if this event is a function callsite event */
 	bool IsFunctionCallSite() const { return (NodeFlags & EScriptExecutionNodeFlags::FunctionCall) != 0U; }
 
@@ -308,7 +353,7 @@ public:
 	/** Gets a typed observed object context */
 	template<typename CastType> const CastType* GetTypedObservedObject() { return Cast<CastType>(ObservedObject.Get()); }
 
-	/** Gets the observed object context */
+	/** Gets the observed pin context */
 	const UEdGraphPin* GetObservedPin() const { return ObservedPin.Get(); }
 
 	/** Navigate to object */
@@ -356,11 +401,8 @@ public:
 	/** Refresh Stats */
 	virtual void RefreshStats(const FTracePath& TracePath);
 
-	/** Calculate heat level stats */
-	virtual void CalculateHeatLevelStats(TSharedPtr<FScriptHeatLevelMetrics> HeatLevelMetrics);
-
-	/** Calculate Hottest Path Stats */
-	virtual float CalculateHottestPathStats(FScriptExecutionHottestPathParams HotPathParams);
+	/** Update heat display statistics */
+	virtual void UpdateHeatDisplayStats(FScriptExecutionHottestPathParams& HotPathParams);
 
 	/** Get all exec nodes */
 	virtual void GetAllExecNodes(TMap<FName, TSharedPtr<FScriptExecutionNode>>& ExecNodesOut);
@@ -428,7 +470,9 @@ public:
 	}
 
 	// FScriptNodePerfData
-	virtual void GetBlueprintPerfDataForAllTracePaths(FScriptPerfData& OutPerfData) override;
+	virtual void UpdatePerfDataForNode() override;
+	virtual void RefreshStats(const FTracePath& TracePath) override;
+	virtual void UpdateHeatDisplayStats(FScriptExecutionHottestPathParams& HotPathParams) override;
 	// ~FScriptNodePerfData
 
 	/** Adds an entry site at the specified script offset */
@@ -461,13 +505,13 @@ public:
 
 	FScriptExecutionTunnelEntry(const FScriptExecNodeParams& InitParams)
 		: FScriptExecutionNode(InitParams)
-		, TunnelEntryCount(0)
 	{
 	}
 
 	// FScriptExecutionNode
 	virtual void GetLinearExecutionPath(TArray<FLinearExecPath>& LinearExecutionNodes, const FTracePath& TracePath, const bool bIncludeChildren = false) override;
 	virtual void RefreshStats(const FTracePath& TracePath) override;
+	virtual void UpdateHeatDisplayStats(FScriptExecutionHottestPathParams& HotPathParams) override;
 	// ~FScriptExecutionNode
 
 	/** Add tunnel timing */
@@ -489,12 +533,10 @@ public:
 	bool IsInternalBoundary(TSharedPtr<FScriptExecutionNode> PotentialBoundaryNode) const;
 
 	/** Increment tunnel entry count */
-	void IncrementTunnelEntryCount() { TunnelEntryCount++; }
+	void IncrementTunnelEntryCount(const FName InstanceName, const FTracePath& TracePath);
 
 private:
 
-	/** The tunnel entry count */
-	int32 TunnelEntryCount;
 	/** The tunnel instance this node represents */
 	TSharedPtr<FScriptExecutionTunnelInstance> TunnelInstance;
 
@@ -536,6 +578,7 @@ public:
 	virtual void RefreshStats(const FTracePath& TracePath) override;
 	virtual TSharedRef<SWidget> GetIconWidget(const uint32 TracePath = 0U) override;
 	virtual TSharedRef<SWidget> GetHyperlinkWidget(const uint32 TracePath = 0U) override;
+	virtual void UpdateHeatDisplayStats(FScriptExecutionHottestPathParams& HotPathParams) override;
 	// ~FScriptExecutionNode
 
 	/** Navigate to exit site using tracepath */
@@ -575,8 +618,7 @@ public:
 	virtual TSharedPtr<FScriptExecutionNode> GetPureChainNode() override { return AsShared(); }
 	virtual void GetLinearExecutionPath(TArray<FLinearExecPath>& LinearExecutionNodes, const FTracePath& TracePath, const bool bIncludeChildren = false) override;
 	virtual void RefreshStats(const FTracePath& TracePath) override;
-	virtual float CalculateHottestPathStats(FScriptExecutionHottestPathParams HotPathParams) override;
-	virtual void CalculateHeatLevelStats(TSharedPtr<FScriptHeatLevelMetrics> HeatLevelMetrics);
+	virtual void UpdateHeatDisplayStats(FScriptExecutionHottestPathParams& HotPathParams) override;
 	// ~FScriptExecutionNode
 
 };
@@ -645,6 +687,12 @@ public:
 	/** Returns the instance that matches the supplied name if present */
 	TSharedPtr<FScriptExecutionNode> GetInstanceByName(FName InstanceName);
 
+	/** Get the active instance name */
+	const FName GetActiveInstanceName() const { return ActiveInstanceName; }
+
+	/** Set the active instance name */
+	void SetActiveInstanceName(const FName InstanceName) { ActiveInstanceName = InstanceName; }
+
 	/** Sorts all contained events */
 	void SortEvents();
 	
@@ -659,11 +707,17 @@ public:
 	// ~FScriptExecutionNode
 
 protected:
+
 	/** Updates internal heat level metrics data */
 	void UpdateHeatLevelMetrics(TSharedPtr<FScriptPerfData> BlueprintData);
 
+	/** Updates heat displays */
+	void UpdateHeatLevels();
+
 private:
 
+	/** Active instance */
+	FName ActiveInstanceName;
 	/** Exec nodes representing all instances based on this blueprint */
 	TArray<TSharedPtr<FScriptExecutionNode>> Instances;
 	/** Internal metrics data used for calculating perf data heat levels */
diff --git a/Engine/Source/Editor/Kismet/Public/Profiler/ScriptPerfData.h b/Engine/Source/Editor/Kismet/Public/Profiler/ScriptPerfData.h
--- a/Engine/Source/Editor/Kismet/Public/Profiler/ScriptPerfData.h
+++ b/Engine/Source/Editor/Kismet/Public/Profiler/ScriptPerfData.h
@@ -78,11 +78,11 @@ public:
 	/** Add data */
 	void AddData(const FScriptPerfData& DataIn);
 
-	/** Initialise from data set */
-	void InitialiseFromDataSet(const TArray<TSharedPtr<FScriptPerfData>>& DataSet);
+	/** Initialise blueprint data from data set */
+	void CreateBlueprintStats(const TArray<TSharedPtr<FScriptPerfData>>& BlueprintDataSet);
 
 	/** Initialise as accumulated from data set */
-	void AccumulateDataSet(const TArray<TSharedPtr<FScriptPerfData>>& DataSet);
+	void AccumulateDataSet(const TArray<TSharedPtr<FScriptPerfData>>& DataSet, bool bApplyAsInclusive = false);
 
 	/** Increments sample count without affecting data */
 	void TickSamples() { RawSamples++; }
@@ -100,10 +100,11 @@ public:
 	// Returns the various stats this container holds
 	double GetAverageTiming() const;
 	double GetInclusiveTiming() const;
-	double GetMaxTiming() const { return MaxTiming; }
-	double GetMinTiming() const { return MinTiming; }
+	double GetMaxTiming() const { return MaxTiming == -MAX_dbl ? 0.0 : MaxTiming; }
+	double GetMinTiming() const { return MinTiming == MAX_dbl ? 0.0 : MinTiming; }
 	double GetTotalTiming() const { return AverageTiming; }
 	float GetSampleCount() const { return RawSamples / SampleFrequency; }
+	float GetRawSampleCount() const { return RawSamples; }
 	void OverrideSampleCount(const int32 NewSampleCount) { RawSamples = NewSampleCount; }
 
 	// Returns performance heat levels for visual display
@@ -113,7 +114,7 @@ public:
 	float GetTotalHeatLevel() const { return TotalTimingHeatLevel; }
 
 	// Calculate performance heat levels for visual display
-	void SetHeatLevels(TSharedPtr<FScriptHeatLevelMetrics> HeatLevelMetrics);
+	void SetHeatLevels(TSharedPtr<const FScriptHeatLevelMetrics> HeatLevelMetrics);
 
 	// Hottest path interface
 	float GetHottestPathHeatLevel() const { return HottestPathHeatValue; }
@@ -126,6 +127,8 @@ public:
 	FText GetMaxTimingText() const;
 	FText GetMinTimingText() const;
 	FText GetSamplesText() const;
+	FText GetHottestPathText() const;
+	FText GetHeatLevelText() const;
 
 private:
 
@@ -160,6 +163,4 @@ private:
 	/** Time number formatting options */
 	static FNumberFormattingOptions TimeNumberFormat;
 
-	/** Sets blueprint level stats as averaged */
-	static bool bAverageBlueprintStats;
 };
diff --git a/Engine/Source/Editor/KismetCompiler/Private/KismetCompiler.cpp b/Engine/Source/Editor/KismetCompiler/Private/KismetCompiler.cpp
--- a/Engine/Source/Editor/KismetCompiler/Private/KismetCompiler.cpp
+++ b/Engine/Source/Editor/KismetCompiler/Private/KismetCompiler.cpp
@@ -215,7 +215,7 @@ void FKismetCompilerContext::CleanAndSanitizeClass(UBlueprintGeneratedClass* Cla
 	// Set properties we need to regenerate the class with
 	ClassToClean->PropertyLink = ParentClass->PropertyLink;
 	ClassToClean->ClassWithin = ParentClass;
-	ClassToClean->ClassConfigName = ClassToClean->IsNative() ? FName(ClassToClean->StaticConfigName()) : FName(*ParentClass->GetConfigName());
+	ClassToClean->ClassConfigName = ClassToClean->IsNative() ? FName(ClassToClean->StaticConfigName()) : ParentClass->ClassConfigName;
 	ClassToClean->DebugData = FBlueprintDebugData();
 }
 
@@ -3086,6 +3086,11 @@ void FKismetCompilerContext::ProcessOneFunctionGraph(UEdGraph* SourceGraph, bool
 	// If a function in the graph cannot be overridden/placed as event make sure that it is not.
 	VerifyValidOverrideFunction(FunctionGraph);
 
+	if (bIsFullCompile)
+	{
+		ValidateGraphIsWellFormed(FunctionGraph);
+	}
+
 	FKismetFunctionContext& Context = *new (FunctionList)FKismetFunctionContext(MessageLog, Schema, NewClass, Blueprint, CompileOptions.DoesRequireCppCodeGeneration(), CompileOptions.IsInstrumentationActive());
 	Context.SourceGraph = FunctionGraph;
 
@@ -3634,6 +3639,10 @@ void FKismetCompilerContext::Compile()
 			BP_SCOPED_COMPILER_EVENT_STAT(EKismetCompilerStats_CodeGenerationTime);
 			const bool bGenerateStubsOnly = !bIsFullCompile || (0 != MessageLog.NumErrors);
 			Backend_VM.GenerateCodeFromClass(NewClass, FunctionList, bGenerateStubsOnly);
+			if (!bGenerateStubsOnly)
+			{
+				Blueprint->bHasAnyNonReducibleFunction = Backend_VM.bAnyNonReducibleFunctionGenerated ? UBlueprint::EIsBPNonReducible::Yes : UBlueprint::EIsBPNonReducible::No;
+			}
 		}
 
 		if (bDisplayBytecode && bIsFullCompile)
diff --git a/Engine/Source/Editor/KismetCompiler/Private/KismetCompilerBackend.h b/Engine/Source/Editor/KismetCompiler/Private/KismetCompilerBackend.h
--- a/Engine/Source/Editor/KismetCompiler/Private/KismetCompilerBackend.h
+++ b/Engine/Source/Editor/KismetCompiler/Private/KismetCompilerBackend.h
@@ -22,6 +22,7 @@ class FKismetCompilerVMBackend : public IKismetCompilerBackend
 {
 public:
 	typedef TMap<FBlueprintCompiledStatement*, CodeSkipSizeType> TStatementToSkipSizeMap;
+
 protected:
 	UBlueprint* Blueprint;
 	UEdGraphSchema_K2* Schema;
@@ -30,11 +31,14 @@ protected:
 
 	TStatementToSkipSizeMap UbergraphStatementLabelMap;
 public:
+	bool bAnyNonReducibleFunctionGenerated;
+
 	FKismetCompilerVMBackend(UBlueprint* InBlueprint, UEdGraphSchema_K2* InSchema, FKismetCompilerContext& InContext)
 		: Blueprint(InBlueprint)
 		, Schema(InSchema)
 		, MessageLog(InContext.MessageLog)
 		, CompilerContext(InContext)
+		, bAnyNonReducibleFunctionGenerated(false)
 	{
 	}
 
diff --git a/Engine/Source/Editor/KismetCompiler/Private/KismetCompilerMisc.cpp b/Engine/Source/Editor/KismetCompiler/Private/KismetCompilerMisc.cpp
--- a/Engine/Source/Editor/KismetCompiler/Private/KismetCompilerMisc.cpp
+++ b/Engine/Source/Editor/KismetCompiler/Private/KismetCompilerMisc.cpp
@@ -72,6 +72,23 @@
 			return false;
 		}
 	}
+	else if (Type.bIsSet)
+	{
+		if (USetProperty* SetProperty = Cast<USetProperty>(Property))
+		{
+			if (OwningFunction && FEdGraphUtilities::IsSetParam(OwningFunction, SourcePin->PinName))
+			{
+				return true;
+			}
+
+			TestProperty = SetProperty->ElementProp;
+		}
+		else
+		{
+			MessageLog.Error(*LOCTEXT("PinSpecifiedAsSet_Error", "Pin @@ is specified as a set, but does not have a valid set property.").ToString(), SourcePin);
+			return false;
+		}
+	}
 	else
 	{
 		// For scalars, we just take the passed in property
@@ -82,7 +99,7 @@
 	if ( (OwningFunction != NULL) && OwningFunction->HasMetaData(FBlueprintMetadata::MD_ArrayParam) )
 	{
 		// Check to see if this param is type dependent on an array parameter
-		const FString DependentParams = OwningFunction->GetMetaData(FBlueprintMetadata::MD_ArrayDependentParam);
+		const FString& DependentParams = OwningFunction->GetMetaData(FBlueprintMetadata::MD_ArrayDependentParam);
 		TArray<FString>	DependentParamNames;
 		DependentParams.ParseIntoArray(DependentParamNames, TEXT(","), true);
 		if (DependentParamNames.Find(SourcePin->PinName) != INDEX_NONE)
@@ -91,6 +108,17 @@
 			return true;
 		}
 	}
+	else if ((OwningFunction != NULL) && OwningFunction->HasMetaData(FBlueprintMetadata::MD_SetParam))
+	{
+		// If the pin in question is part of a Set (inferred) parameter, then ignore pin matching:
+		// @todo:  This assumes that the wildcard coersion has done its job...I'd feel better if 
+		// there was some easier way of accessing the target array type
+		if (FEdGraphUtilities::IsSetParam(OwningFunction, SourcePin->PinName))
+		{
+			return true;
+		}
+	}
 
+
 	int32 NumErrorsAtStart = MessageLog.NumErrors;
 
@@ -747,59 +775,17 @@
 	}
 }
 
-/** Creates a property named PropertyName of type PropertyType in the Scope or returns NULL if the type is unknown, but does *not* link that property in */
-UProperty* FKismetCompilerUtilities::CreatePropertyOnScope(UStruct* Scope, const FName& PropertyName, const FEdGraphPinType& Type, UClass* SelfClass, uint64 PropertyFlags, const UEdGraphSchema_K2* Schema, FCompilerResultsLog& MessageLog)
+UProperty* FKismetCompilerUtilities::CreatePrimitiveProperty(UObject* PropertyScope, const FName& ValidatedPropertyName, const FString& PinCategory, const FString& PinSubCategory, UObject* PinSubCategoryObject, UClass* SelfClass, bool bIsWeakPointer, const class UEdGraphSchema_K2* Schema, FCompilerResultsLog& MessageLog)
 {
-	//@TODO: Check for name conflicts!
-
-	// Properties are non-transactional as they're regenerated on every compile
 	const EObjectFlags ObjectFlags = RF_Public;
 
-	UProperty* NewProperty = NULL;
-	UObject* PropertyScope = NULL;
-
-	FName ValidatedPropertyName = PropertyName;
-
-	// Check to see if there's already a object on this scope with the same name, and throw an internal compiler error if so
-	// If this happens, it breaks the property link, which causes stack corruption and hard-to-track errors, so better to fail at this point
-	{
-		if (UObject* ExistingObject = CheckPropertyNameOnScope(Scope, PropertyName))
-		{
-			MessageLog.Error(*FString::Printf(TEXT("Internal Compiler Error:  Tried to create a property %s in scope %s, but %s already exists there."), *PropertyName.ToString(), (Scope ? *Scope->GetName() : TEXT("None")), *ExistingObject->GetFullName()));
-
-			// Find a free name, so we can still create the property to make it easier to spot the duplicates, and avoid crashing
-			uint32 Counter = 0;
-			FName TestName;
-			do 
-			{
-				FString TestNameString = PropertyName.ToString() + FString::Printf(TEXT("_ERROR_DUPLICATE_%d"), Counter++);
-				TestName = FName(*TestNameString);
-
-			} while (CheckPropertyNameOnScope(Scope, TestName) != NULL);
-
-			ValidatedPropertyName = TestName;
-		}
-	}
-
-	// Handle creating an array property, if necessary
-	const bool bIsArrayProperty = Type.bIsArray;
-	UArrayProperty* NewArrayProperty = NULL;
-	if( bIsArrayProperty )
-	{
-		NewArrayProperty = NewObject<UArrayProperty>(Scope, ValidatedPropertyName, ObjectFlags);
-		PropertyScope = NewArrayProperty;
-	}
-	else
-	{
-		PropertyScope = Scope;
-	}
-
+	UProperty* NewProperty = nullptr;
 	//@TODO: Nasty string if-else tree
-	if ((Type.PinCategory == Schema->PC_Object) || (Type.PinCategory == Schema->PC_Interface) || (Type.PinCategory == Schema->PC_Asset))
+	if ((PinCategory == Schema->PC_Object) || (PinCategory == Schema->PC_Interface) || (PinCategory == Schema->PC_Asset))
 	{
-		UClass* SubType = (Type.PinSubCategory == Schema->PSC_Self) ? SelfClass : Cast<UClass>(Type.PinSubCategoryObject.Get());
+		UClass* SubType = (PinSubCategory == Schema->PSC_Self) ? SelfClass : Cast<UClass>(PinSubCategoryObject);
 
-		if( SubType == NULL )
+		if (SubType == NULL)
 		{
 			// If this is from a degenerate pin, because the object type has been removed, default this to a UObject subtype so we can make a dummy term for it to allow the compiler to continue
 			SubType = UObject::StaticClass();
@@ -823,11 +809,11 @@
 			{
 				UObjectPropertyBase* NewPropertyObj = NULL;
 
-				if (Type.PinCategory == Schema->PC_Asset)
+				if (PinCategory == Schema->PC_Asset)
 				{
 					NewPropertyObj = NewObject<UAssetObjectProperty>(PropertyScope, ValidatedPropertyName, ObjectFlags);
 				}
-				else if( Type.bIsWeakPointer )
+				else if (bIsWeakPointer)
 				{
 					NewPropertyObj = NewObject<UWeakObjectProperty>(PropertyScope, ValidatedPropertyName, ObjectFlags);
 				}
@@ -843,9 +829,9 @@
 			}
 		}
 	}
-	else if (Type.PinCategory == Schema->PC_Struct)
+	else if (PinCategory == Schema->PC_Struct)
 	{
-		UScriptStruct* SubType = Cast<UScriptStruct>(Type.PinSubCategoryObject.Get());
+		UScriptStruct* SubType = Cast<UScriptStruct>(PinSubCategoryObject);
 		if (SubType != NULL)
 		{
 			FString StructureError;
@@ -858,10 +844,6 @@
 				if (SubType->StructFlags & STRUCT_HasInstancedReference)
 				{
 					NewProperty->SetPropertyFlags(CPF_ContainsInstancedReference);
-					if (NewArrayProperty)
-					{
-						NewArrayProperty->SetPropertyFlags(CPF_ContainsInstancedReference);
-					}
 				}
 			}
 			else
@@ -869,17 +851,17 @@
 				MessageLog.Error(
 					*FString::Printf(
 						*LOCTEXT("InvalidStructForField_Error", "Invalid property '%s' structure '%s' error: %s").ToString(),
-						*PropertyName.ToString(),
+						*ValidatedPropertyName.ToString(),
 						*SubType->GetName(),
 						*StructureError
-					));
+						));
 			}
 		}
 	}
-	else if ((Type.PinCategory == Schema->PC_Class) || (Type.PinCategory == Schema->PC_AssetClass))
+	else if ((PinCategory == Schema->PC_Class) || (PinCategory == Schema->PC_AssetClass))
 	{
-		UClass* SubType = Cast<UClass>(Type.PinSubCategoryObject.Get());
-		
+		UClass* SubType = Cast<UClass>(PinSubCategoryObject);
+
 		if (SubType == NULL)
 		{
 			// If this is from a degenerate pin, because the object type has been removed, default this to a UObject subtype so we can make a dummy term for it to allow the compiler to continue
@@ -887,14 +869,14 @@
 
 			MessageLog.Warning(
 				*FString::Printf(
-				*LOCTEXT("InvalidClassForField_Error", "Invalid property '%s' class, replaced with Object.  Please fix or remove.").ToString(),
-				*PropertyName.ToString()
-				));
+					*LOCTEXT("InvalidClassForField_Error", "Invalid property '%s' class, replaced with Object.  Please fix or remove.").ToString(),
+					*ValidatedPropertyName.ToString()
+					));
 		}
 
 		if (SubType != NULL)
 		{
-			if (Type.PinCategory == Schema->PC_AssetClass)
+			if (PinCategory == Schema->PC_AssetClass)
 			{
 				auto AssetClassProperty = NewObject<UAssetClassProperty>(PropertyScope, ValidatedPropertyName, ObjectFlags);
 				// we want to use this setter function instead of setting the 
@@ -916,52 +898,36 @@
 			}
 		}
 	}
-	else if (Type.PinCategory == Schema->PC_Delegate)
+	else if (PinCategory == Schema->PC_Int)
 	{
-		if (UFunction* SignatureFunction = FMemberReference::ResolveSimpleMemberReference<UFunction>(Type.PinSubCategoryMemberReference))
-		{
-			UDelegateProperty* NewPropertyDelegate = NewObject<UDelegateProperty>(PropertyScope, ValidatedPropertyName, ObjectFlags);
-			NewPropertyDelegate->SignatureFunction = SignatureFunction;
-			NewProperty = NewPropertyDelegate;
-		}
-	}
-	else if (Type.PinCategory == Schema->PC_MCDelegate)
-	{
-		UFunction* const SignatureFunction = FMemberReference::ResolveSimpleMemberReference<UFunction>(Type.PinSubCategoryMemberReference);
-		UMulticastDelegateProperty* NewPropertyDelegate = NewObject<UMulticastDelegateProperty>(PropertyScope, ValidatedPropertyName, ObjectFlags);
-		NewPropertyDelegate->SignatureFunction = SignatureFunction;
-		NewProperty = NewPropertyDelegate;
-	}
-	else if (Type.PinCategory == Schema->PC_Int)
-	{
 		NewProperty = NewObject<UIntProperty>(PropertyScope, ValidatedPropertyName, ObjectFlags);
 	}
-	else if (Type.PinCategory == Schema->PC_Float)
+	else if (PinCategory == Schema->PC_Float)
 	{
 		NewProperty = NewObject<UFloatProperty>(PropertyScope, ValidatedPropertyName, ObjectFlags);
 	}
-	else if (Type.PinCategory == Schema->PC_Boolean)
+	else if (PinCategory == Schema->PC_Boolean)
 	{
 		UBoolProperty* BoolProperty = NewObject<UBoolProperty>(PropertyScope, ValidatedPropertyName, ObjectFlags);
 		BoolProperty->SetBoolSize(sizeof(bool), true);
 		NewProperty = BoolProperty;
 	}
-	else if (Type.PinCategory == Schema->PC_String)
+	else if (PinCategory == Schema->PC_String)
 	{
 		NewProperty = NewObject<UStrProperty>(PropertyScope, ValidatedPropertyName, ObjectFlags);
 	}
-	else if (Type.PinCategory == Schema->PC_Text)
+	else if (PinCategory == Schema->PC_Text)
 	{
 		NewProperty = NewObject<UTextProperty>(PropertyScope, ValidatedPropertyName, ObjectFlags);
 	}
-	else if (Type.PinCategory == Schema->PC_Byte)
+	else if (PinCategory == Schema->PC_Byte)
 	{
 		UByteProperty* ByteProp = NewObject<UByteProperty>(PropertyScope, ValidatedPropertyName, ObjectFlags);
-		ByteProp->Enum = Cast<UEnum>(Type.PinSubCategoryObject.Get());
+		ByteProp->Enum = Cast<UEnum>(PinSubCategoryObject);
 
 		NewProperty = ByteProp;
 	}
-	else if (Type.PinCategory == Schema->PC_Name)
+	else if (PinCategory == Schema->PC_Name)
 	{
 		NewProperty = NewObject<UNameProperty>(PropertyScope, ValidatedPropertyName, ObjectFlags);
 	}
@@ -971,7 +937,145 @@
 		NewProperty = NewObject<UIntProperty>(PropertyScope, ValidatedPropertyName, ObjectFlags);
 	}
 
-	if (bIsArrayProperty)
+	return NewProperty;
+}
+
+/** Creates a property named PropertyName of type PropertyType in the Scope or returns NULL if the type is unknown, but does *not* link that property in */
+UProperty* FKismetCompilerUtilities::CreatePropertyOnScope(UStruct* Scope, const FName& PropertyName, const FEdGraphPinType& Type, UClass* SelfClass, uint64 PropertyFlags, const UEdGraphSchema_K2* Schema, FCompilerResultsLog& MessageLog)
+{
+	const EObjectFlags ObjectFlags = RF_Public;
+
+	UProperty* NewProperty = NULL;
+	UObject* PropertyScope = NULL;
+
+	FName ValidatedPropertyName = PropertyName;
+
+	// Check to see if there's already a object on this scope with the same name, and throw an internal compiler error if so
+	// If this happens, it breaks the property link, which causes stack corruption and hard-to-track errors, so better to fail at this point
+	{
+		if (UObject* ExistingObject = CheckPropertyNameOnScope(Scope, PropertyName))
+		{
+			MessageLog.Error(*FString::Printf(TEXT("Internal Compiler Error:  Tried to create a property %s in scope %s, but %s already exists there."), *PropertyName.ToString(), (Scope ? *Scope->GetName() : TEXT("None")), *ExistingObject->GetFullName()));
+
+			// Find a free name, so we can still create the property to make it easier to spot the duplicates, and avoid crashing
+			uint32 Counter = 0;
+			FName TestName;
+			do 
+			{
+				FString TestNameString = PropertyName.ToString() + FString::Printf(TEXT("_ERROR_DUPLICATE_%d"), Counter++);
+				TestName = FName(*TestNameString);
+
+			} while (CheckPropertyNameOnScope(Scope, TestName) != NULL);
+
+			ValidatedPropertyName = TestName;
+		}
+	}
+
+	// Handle creating an array property, if necessary
+	const bool bIsMapProperty = Type.bIsMap;
+	const bool bIsSetProperty = Type.bIsSet;
+	const bool bIsArrayProperty = Type.bIsArray;
+	// only one of these should be set (or none):
+	ensure( !(bIsMapProperty || bIsSetProperty || bIsArrayProperty) ||
+			( !(bIsMapProperty && bIsSetProperty && bIsArrayProperty) && (bIsMapProperty ^ bIsSetProperty ^ bIsArrayProperty) )
+			);
+	UMapProperty* NewMapProperty = nullptr;
+	USetProperty* NewSetProperty = nullptr;
+	UArrayProperty* NewArrayProperty = nullptr;
+	UProperty* NewContainerProperty = nullptr;
+	if (bIsMapProperty)
+	{
+		NewMapProperty = NewObject<UMapProperty>(Scope, ValidatedPropertyName, ObjectFlags);
+		PropertyScope = NewMapProperty;
+		NewContainerProperty = NewMapProperty;
+	}
+	else if (bIsSetProperty)
+	{
+		NewSetProperty = NewObject<USetProperty>(Scope, ValidatedPropertyName, ObjectFlags);
+		PropertyScope = NewSetProperty;
+		NewContainerProperty = NewSetProperty;
+	}
+	else if( bIsArrayProperty )
+	{
+		NewArrayProperty = NewObject<UArrayProperty>(Scope, ValidatedPropertyName, ObjectFlags);
+		PropertyScope = NewArrayProperty;
+		NewContainerProperty = NewArrayProperty;
+	}
+	else
+	{
+		PropertyScope = Scope;
+	}
+
+	if (Type.PinCategory == Schema->PC_Delegate)
+	{
+		if (UFunction* SignatureFunction = FMemberReference::ResolveSimpleMemberReference<UFunction>(Type.PinSubCategoryMemberReference))
+		{
+			UDelegateProperty* NewPropertyDelegate = NewObject<UDelegateProperty>(PropertyScope, ValidatedPropertyName, ObjectFlags);
+			NewPropertyDelegate->SignatureFunction = SignatureFunction;
+			NewProperty = NewPropertyDelegate;
+		}
+	}
+	else if (Type.PinCategory == Schema->PC_MCDelegate)
+	{
+		UFunction* const SignatureFunction = FMemberReference::ResolveSimpleMemberReference<UFunction>(Type.PinSubCategoryMemberReference);
+		UMulticastDelegateProperty* NewPropertyDelegate = NewObject<UMulticastDelegateProperty>(PropertyScope, ValidatedPropertyName, ObjectFlags);
+		NewPropertyDelegate->SignatureFunction = SignatureFunction;
+		NewProperty = NewPropertyDelegate;
+	}
+	else
+	{
+		NewProperty = CreatePrimitiveProperty(PropertyScope, ValidatedPropertyName, Type.PinCategory, Type.PinSubCategory, Type.PinSubCategoryObject.Get(), SelfClass, Type.bIsWeakPointer, Schema, MessageLog);
+	}
+
+	if (NewContainerProperty && NewProperty && NewProperty->HasAnyPropertyFlags(CPF_ContainsInstancedReference))
+	{
+		NewContainerProperty->SetPropertyFlags(CPF_ContainsInstancedReference);
+	}
+
+	if (bIsMapProperty)
 	{
 		if (NewProperty)
 		{
+			// make the value property:
+			// not feelign good about myself..
+			// Fix up the array property to have the new type-specific property as its inner, and return the new UArrayProperty
+			NewMapProperty->KeyProp = NewProperty;
+			// make sure the value property does not collide with the key property:
+			FName KeyName = FName( *(ValidatedPropertyName.GetPlainNameString() + FString(TEXT("_Key") )) );
+			NewMapProperty->ValueProp = CreatePrimitiveProperty(PropertyScope, KeyName, Type.PinValueType.TerminalCategory, Type.PinValueType.TerminalSubCategory, Type.PinValueType.TerminalSubCategoryObject.Get(), SelfClass, Type.bIsWeakPointer, Schema, MessageLog);;
+			if (!NewMapProperty->ValueProp)
+			{
+				NewMapProperty->MarkPendingKill();
+			}
+			else
+			{
+				if (NewMapProperty->ValueProp && NewMapProperty->ValueProp->HasAnyPropertyFlags(CPF_ContainsInstancedReference))
+				{
+					NewMapProperty->ValueProp->SetPropertyFlags(CPF_ContainsInstancedReference);
+				}
+
+				NewProperty = NewMapProperty;
+			}
+		}
+		else
+		{
+			NewMapProperty->MarkPendingKill();
+		}
+	}
+	else if (bIsSetProperty)
+	{
+		if (NewProperty)
+		{
+			NewProperty->SetPropertyFlags(CPF_HasGetValueTypeHash);
+			NewSetProperty->ElementProp = NewProperty;
+			NewProperty = NewSetProperty;
+		}
+		else
+		{
+			NewSetProperty->MarkPendingKill();
+		}
+	}
+	else if (bIsArrayProperty)
+	{
+		if (NewProperty)
+		{
@@ -1230,6 +1334,25 @@
 	}
 }
 
+bool FKismetCompilerUtilities::IsStatementReducible(EKismetCompiledStatementType StatementType)
+{
+	switch (StatementType)
+	{
+	case EKismetCompiledStatementType::KCST_Nop:
+	case EKismetCompiledStatementType::KCST_UnconditionalGoto:
+	case EKismetCompiledStatementType::KCST_ComputedGoto:
+	case EKismetCompiledStatementType::KCST_Return:
+	case EKismetCompiledStatementType::KCST_EndOfThread:
+	case EKismetCompiledStatementType::KCST_Comment:
+	case EKismetCompiledStatementType::KCST_DebugSite:
+	case EKismetCompiledStatementType::KCST_WireTraceSite:
+	case EKismetCompiledStatementType::KCST_GotoReturn:
+	case EKismetCompiledStatementType::KCST_AssignmentOnPersistentFrame:
+		return true;
+	}
+	return false;
+}
+
 //////////////////////////////////////////////////////////////////////////
 // FNodeHandlingFunctor
 
diff --git a/Engine/Source/Editor/KismetCompiler/Private/KismetCompilerVMBackend.cpp b/Engine/Source/Editor/KismetCompiler/Private/KismetCompilerVMBackend.cpp
--- a/Engine/Source/Editor/KismetCompiler/Private/KismetCompilerVMBackend.cpp
+++ b/Engine/Source/Editor/KismetCompiler/Private/KismetCompilerVMBackend.cpp
@@ -1897,6 +1897,11 @@ void FKismetCompilerVMBackend::ConstructFunction(FKismetFunctionContext& Functio
 					FBlueprintCompiledStatement* Statement = (*StatementList)[StatementIndex];
 
 					ScriptWriter.GenerateCodeForStatement(CompilerContext, FunctionContext, *Statement, StatementNode);
+					
+					const bool bUberGraphFunctionCall = Statement->FunctionToCall && (Statement->FunctionToCall == Class->UberGraphFunction)
+						&& (EKismetCompiledStatementType::KCST_CallFunction == Statement->Type);
+					const bool bIsReducible = FKismetCompilerUtilities::IsStatementReducible(Statement->Type) || bUberGraphFunctionCall;
+					bAnyNonReducibleFunctionGenerated |= !bIsReducible;
 				}
 			}
 		}
@@ -1904,7 +1909,7 @@ void FKismetCompilerVMBackend::ConstructFunction(FKismetFunctionContext& Functio
 
 	// Handle the function return value
 	ScriptWriter.GenerateCodeForStatement(CompilerContext, FunctionContext, ReturnStatement, NULL);	
-	
+
 	// Fix up jump addresses
 	ScriptWriter.PerformFixups();
 
diff --git a/Engine/Source/Editor/KismetCompiler/Public/KismetCompilerMisc.h b/Engine/Source/Editor/KismetCompiler/Public/KismetCompilerMisc.h
--- a/Engine/Source/Editor/KismetCompiler/Public/KismetCompilerMisc.h
+++ b/Engine/Source/Editor/KismetCompiler/Public/KismetCompilerMisc.h
@@ -60,6 +60,9 @@ public:
 	/** return function, that overrides BlueprintImplementableEvent with given name in given class (super-classes are not considered) */
 	static const UFunction* FindOverriddenImplementableEvent(const FName& EventName, const UClass* Class);
 
+	/** Helper function for creating property for primitive types. Used only to create inner peroperties for UArrayProperty, USetProperty, and UMapProperty: */
+	static UProperty* CreatePrimitiveProperty( UObject* PropertyScope, const FName& ValidatedPropertyName, const FString& PinCategory, const FString& PinSubCategory, UObject* PinSubCategoryObject, UClass* SelfClass, bool bIsWeakPointer, const class UEdGraphSchema_K2* Schema, FCompilerResultsLog& MessageLog);
+
 	/** Creates a property named PropertyName of type PropertyType in the Scope or returns NULL if the type is unknown, but does *not* link that property in */
 	static UProperty* CreatePropertyOnScope(UStruct* Scope, const FName& PropertyName, const FEdGraphPinType& Type, UClass* SelfClass, uint64 PropertyFlags, const class UEdGraphSchema_K2* Schema, FCompilerResultsLog& MessageLog);
 
@@ -100,6 +103,8 @@ public:
 
 	/** Generate an error for non-const output parameters */
 	static void DetectValuesReturnedByRef(const UFunction* Func, const UK2Node * Node, FCompilerResultsLog& MessageLog);
+
+	static bool IsStatementReducible(EKismetCompiledStatementType StatementType);
 };
 
 //////////////////////////////////////////////////////////////////////////
