Submitting on behalf of Maciej.Mroz:
Fixing CIS errors - making so we also prune pure nodes that are unused from Blueprint graphs (on compile).

diff --git a/Engine/Source/Editor/KismetCompiler/Private/KismetCompiler.cpp b/Engine/Source/Editor/KismetCompiler/Private/KismetCompiler.cpp
--- a/Engine/Source/Editor/KismetCompiler/Private/KismetCompiler.cpp
+++ b/Engine/Source/Editor/KismetCompiler/Private/KismetCompiler.cpp
@@ -1116,7 +1116,7 @@ void FKismetCompilerContext::TransformNodes(FKismetFunctionContext& Context)
 	}
 }
 
-
+/** Use to traverse exec wires to identify impure (exec) nodes that are used (and shouldn't be pruned) */
 struct FNodeVisitorDownExecWires
 {
 	TSet<UEdGraphNode*> VisitedNodes;
@@ -1155,6 +1155,45 @@ struct FNodeVisitorDownExecWires
 	}
 };
 
+/** Use to traverse data wires (out from exec nodes) to identify pure nodes that are used (and shouldn't be pruned) */
+struct FNodeVisitorUpDataWires
+{
+	TSet<UEdGraphNode*> VisitedNodes;
+	UEdGraphSchema_K2* Schema;
+
+	void TraverseNodes(UEdGraphNode* Node)
+	{
+		bool bAlreadyVisited = false;
+		VisitedNodes.Add(Node, &bAlreadyVisited);
+		if (!bAlreadyVisited)
+		{
+			// Follow every data input pin
+			// we don't have to worry about unconnected non-pure nodes, thay were already removed
+			// we want to gather all pure nodes, that are really used
+			for (int32 i = 0; i < Node->Pins.Num(); ++i)
+			{
+				UEdGraphPin* MyPin = Node->Pins[i];
+
+				if ((MyPin->Direction == EGPD_Input) && !Schema->IsExecPin(*MyPin))
+				{
+					for (int32 j = 0; j < MyPin->LinkedTo.Num(); ++j)
+					{
+						UEdGraphPin* OtherPin = MyPin->LinkedTo[j];
+						if (OtherPin)
+						{
+							UEdGraphNode* OtherNode = OtherPin->GetOwningNode();
+							if (!VisitedNodes.Contains(OtherNode))
+							{
+								TraverseNodes(OtherNode);
+							}
+						}
+					}
+				}
+			}
+		}
+	}
+};
+
 bool FKismetCompilerContext::CanIgnoreNode(const UEdGraphNode* Node) const
 {
 	if (const UK2Node* K2Node = Cast<const UK2Node>(Node))
@@ -1238,6 +1277,38 @@ void FKismetCompilerContext::PruneIsolatedNodes(const TArray<UEdGraphNode*>& Roo
 
 	// Prune the nodes that aren't even reachable via data dependencies
 	Super::PruneIsolatedNodes(RootSet, GraphNodes);
+
+	{
+		FNodeVisitorUpDataWires UpDataVisitor;
+		UpDataVisitor.Schema = Schema;
+		// we still have pure nodes that could afford to be pruned, so let's 
+		// explore data wires (from the impure nodes we kept), and identify
+		// pure nodes we want to keep
+		for (UEdGraphNode* VisitedNode : Visitor.VisitedNodes)
+		{
+			UK2Node* K2Node = Cast<UK2Node>(VisitedNode);
+			if (K2Node && !K2Node->IsNodePure())
+			{
+				UpDataVisitor.TraverseNodes(VisitedNode);
+			}
+		}
+
+		// remove pure nodes that are unused (ones that weren't visited by traversing data wires)
+		for (int32 NodeIndex = 0; NodeIndex < GraphNodes.Num(); ++NodeIndex)
+		{
+			UK2Node* K2Node = Cast<UK2Node>(GraphNodes[NodeIndex]);
+			if (K2Node && K2Node->IsNodePure() && !UpDataVisitor.VisitedNodes.Contains(K2Node))
+				//TODO: Add case for exec knot node
+			{
+				if (!ShouldForceKeepNode(K2Node))
+				{
+					K2Node->BreakAllNodeLinks();
+					GraphNodes.RemoveAtSwap(NodeIndex);
+					--NodeIndex;
+				}
+			}
+		}
+	}
 }
 
 /**
