Knot node are not pruned like regular pure node, because Knots can have exec pins.
Merged cl3153986

diff --git a/Engine/Source/Editor/KismetCompiler/Private/KismetCompiler.cpp b/Engine/Source/Editor/KismetCompiler/Private/KismetCompiler.cpp
--- a/Engine/Source/Editor/KismetCompiler/Private/KismetCompiler.cpp
+++ b/Engine/Source/Editor/KismetCompiler/Private/KismetCompiler.cpp
@@ -1297,8 +1297,8 @@ void FKismetCompilerContext::PruneIsolatedNodes(const TArray<UEdGraphNode*>& Roo
 		for (int32 NodeIndex = 0; NodeIndex < GraphNodes.Num(); ++NodeIndex)
 		{
 			UK2Node* K2Node = Cast<UK2Node>(GraphNodes[NodeIndex]);
-			if (K2Node && K2Node->IsNodePure() && !UpDataVisitor.VisitedNodes.Contains(K2Node))
-				//TODO: Add case for exec knot node
+			if (K2Node && K2Node->IsNodePure() && !UpDataVisitor.VisitedNodes.Contains(K2Node) 
+				&& !K2Node->IsA<UK2Node_Knot>()) // Knots are pure, but they can have exec pins
 			{
 				if (!ShouldForceKeepNode(K2Node))
 				{
