Fixing some localization key conflicts that crept in from stream merges

diff --git a/Engine/Source/Developer/AssetTools/Private/AssetTypeActions/AssetTypeActions_ForceFeedbackEffect.cpp b/Engine/Source/Developer/AssetTools/Private/AssetTypeActions/AssetTypeActions_ForceFeedbackEffect.cpp
--- a/Engine/Source/Developer/AssetTools/Private/AssetTypeActions/AssetTypeActions_ForceFeedbackEffect.cpp
+++ b/Engine/Source/Developer/AssetTools/Private/AssetTypeActions/AssetTypeActions_ForceFeedbackEffect.cpp
@@ -176,7 +176,7 @@ TSharedPtr<SWidget> FAssetTypeActions_ForceFeedbackEffect::GetThumbnailOverlay(c
 			return LOCTEXT("Thumbnail_StopForceFeedbackToolTip", "Stop selected force feedback effect");
 		}
 
-		return LOCTEXT("Thumbnail_PlaySoundToolTip", "Play selected force feedback effect");
+		return LOCTEXT("Thumbnail_PlayForceFeedbackToolTip", "Play selected force feedback effect");
 	};
 
 	return SNew(SBox)
diff --git a/Engine/Source/Editor/AnimGraph/Private/AnimGraphNode_CurveSource.cpp b/Engine/Source/Editor/AnimGraph/Private/AnimGraphNode_CurveSource.cpp
--- a/Engine/Source/Editor/AnimGraph/Private/AnimGraphNode_CurveSource.cpp
+++ b/Engine/Source/Editor/AnimGraph/Private/AnimGraphNode_CurveSource.cpp
@@ -18,7 +18,7 @@ FText UAnimGraphNode_CurveSource::GetNodeTitle(ENodeTitleType::Type TitleType) c
 {
 	if (TitleType != ENodeTitleType::MenuTitle && Node.SourceBinding != NAME_None)
 	{
-		return FText::Format(LOCTEXT("AnimGraphNode_CurveSource_Title", "Curve Source: {0}"), FText::FromName(Node.SourceBinding));
+		return FText::Format(LOCTEXT("AnimGraphNode_CurveSource_Title_Fmt", "Curve Source: {0}"), FText::FromName(Node.SourceBinding));
 	}
 	else
 	{
diff --git a/Engine/Source/Editor/BlueprintGraph/Private/EdGraphSchema_K2.cpp b/Engine/Source/Editor/BlueprintGraph/Private/EdGraphSchema_K2.cpp
--- a/Engine/Source/Editor/BlueprintGraph/Private/EdGraphSchema_K2.cpp
+++ b/Engine/Source/Editor/BlueprintGraph/Private/EdGraphSchema_K2.cpp
@@ -3503,7 +3503,7 @@ FText UEdGraphSchema_K2::TypeToText(UProperty* const Property)
 		{
 			FFormatNamedArguments Args;
 			Args.Add(TEXT("SetType"), TypeToText(Set->ElementProp));
-			return FText::Format(LOCTEXT("ArrayPropertyText", "Set of {SetType}"), Args);
+			return FText::Format(LOCTEXT("SetPropertyText", "Set of {SetType}"), Args);
 		}
 	}
 	else if (UMapProperty* Map = Cast<UMapProperty>(Property))
@@ -3513,7 +3513,7 @@ FText UEdGraphSchema_K2::TypeToText(UProperty* const Property)
 			FFormatNamedArguments Args;
 			Args.Add(TEXT("MapKeyType"), TypeToText(Map->KeyProp));
 			Args.Add(TEXT("MapValueType"), TypeToText(Map->ValueProp));
-			return FText::Format(LOCTEXT("ArrayPropertyText", "Map of {MapKeyType} to {MapValueType}"), Args);
+			return FText::Format(LOCTEXT("MapPropertyText", "Map of {MapKeyType} to {MapValueType}"), Args);
 		}
 	}
 	
diff --git a/Engine/Source/Editor/KismetCompiler/Private/AnimBlueprintCompiler.cpp b/Engine/Source/Editor/KismetCompiler/Private/AnimBlueprintCompiler.cpp
--- a/Engine/Source/Editor/KismetCompiler/Private/AnimBlueprintCompiler.cpp
+++ b/Engine/Source/Editor/KismetCompiler/Private/AnimBlueprintCompiler.cpp
@@ -1842,12 +1842,12 @@ void FAnimBlueprintCompiler::PostCompile()
 
 									if (FunctionNode)
 									{
-										MessageLog.Warning(*LOCTEXT("NotThreadSafeWarning", "Node @@ uses potentially thread-unsafe call @@. Disable threaded update or use a thread-safe call. Function may need BlueprintThreadSafe metadata adding.").ToString(), SourceNode, FunctionNode)
+										MessageLog.Warning(*LOCTEXT("NotThreadSafeWarningNodeContext", "Node @@ uses potentially thread-unsafe call @@. Disable threaded update or use a thread-safe call. Function may need BlueprintThreadSafe metadata adding.").ToString(), SourceNode, FunctionNode)
 											->AddToken(FDocumentationToken::Create(TEXT("Engine/Animation/AnimBlueprints/AnimGraph")));
 									}
 									else if(Statement->FunctionToCall)
 									{
-										MessageLog.Warning(*FText::Format(LOCTEXT("NotThreadSafeWarningUnknownContext", "Node @@ uses potentially thread-unsafe call {0}. Disable threaded update or use a thread-safe call. Function may need BlueprintThreadSafe metadata adding."), Statement->FunctionToCall->GetDisplayNameText()).ToString(), SourceNode)
+										MessageLog.Warning(*FText::Format(LOCTEXT("NotThreadSafeWarningFunctionContext", "Node @@ uses potentially thread-unsafe call {0}. Disable threaded update or use a thread-safe call. Function may need BlueprintThreadSafe metadata adding."), Statement->FunctionToCall->GetDisplayNameText()).ToString(), SourceNode)
 											->AddToken(FDocumentationToken::Create(TEXT("Engine/Animation/AnimBlueprints/AnimGraph")));
 									}
 									else
diff --git a/Engine/Source/Runtime/AIModule/Private/EnvironmentQuery/Generators/EnvQueryGenerator_Cone.cpp b/Engine/Source/Runtime/AIModule/Private/EnvironmentQuery/Generators/EnvQueryGenerator_Cone.cpp
--- a/Engine/Source/Runtime/AIModule/Private/EnvironmentQuery/Generators/EnvQueryGenerator_Cone.cpp
+++ b/Engine/Source/Runtime/AIModule/Private/EnvironmentQuery/Generators/EnvQueryGenerator_Cone.cpp
@@ -84,7 +84,7 @@ void UEnvQueryGenerator_Cone::GenerateItems(FEnvQueryInstance& QueryInstance) co
 
 FText UEnvQueryGenerator_Cone::GetDescriptionTitle() const
 {
-	return FText::Format(LOCTEXT("DescriptionGenerateAroundContext", "{0}: generate in front of {1}"),
+	return FText::Format(LOCTEXT("ConeDescriptionGenerateAroundContext", "{0}: generate in front of {1}"),
 		Super::GetDescriptionTitle(), UEnvQueryTypes::DescribeContext(CenterActor));
 }
 
diff --git a/Engine/Source/Runtime/AIModule/Private/EnvironmentQuery/Generators/EnvQueryGenerator_SimpleGrid.cpp b/Engine/Source/Runtime/AIModule/Private/EnvironmentQuery/Generators/EnvQueryGenerator_SimpleGrid.cpp
--- a/Engine/Source/Runtime/AIModule/Private/EnvironmentQuery/Generators/EnvQueryGenerator_SimpleGrid.cpp
+++ b/Engine/Source/Runtime/AIModule/Private/EnvironmentQuery/Generators/EnvQueryGenerator_SimpleGrid.cpp
@@ -48,7 +48,7 @@ void UEnvQueryGenerator_SimpleGrid::GenerateItems(FEnvQueryInstance& QueryInstan
 
 FText UEnvQueryGenerator_SimpleGrid::GetDescriptionTitle() const
 {
-	return FText::Format(LOCTEXT("DescriptionGenerateAroundContext", "{0}: generate around {1}"),
+	return FText::Format(LOCTEXT("SimpleGridDescriptionGenerateAroundContext", "{0}: generate around {1}"),
 		Super::GetDescriptionTitle(), UEnvQueryTypes::DescribeContext(GenerateAround));
 };
 
