Fix for various non-unity CIS errors introduced by the Dev-BP merge (IWYU yo!).

diff --git a/Engine/Source/Editor/BlueprintGraph/Private/K2Node_MakeVariable.cpp b/Engine/Source/Editor/BlueprintGraph/Private/K2Node_MakeVariable.cpp
--- a/Engine/Source/Editor/BlueprintGraph/Private/K2Node_MakeVariable.cpp
+++ b/Engine/Source/Editor/BlueprintGraph/Private/K2Node_MakeVariable.cpp
@@ -7,7 +7,7 @@
 #include "EdGraphUtilities.h"
 #include "EdGraph/EdGraphNodeUtils.h"
 #include "Kismet2/BlueprintEditorUtils.h"
-
+#include "Templates/SharedPointer.h"
 #include "KismetCompiledFunctionContext.h"
 #include "KismetCompiler.h"
 #include "KismetCompilerMisc.h"
diff --git a/Engine/Source/Editor/KismetCompiler/Public/BlueprintCompilerCppBackendInterface.h b/Engine/Source/Editor/KismetCompiler/Public/BlueprintCompilerCppBackendInterface.h
--- a/Engine/Source/Editor/KismetCompiler/Public/BlueprintCompilerCppBackendInterface.h
+++ b/Engine/Source/Editor/KismetCompiler/Public/BlueprintCompilerCppBackendInterface.h
@@ -6,6 +6,7 @@
 #include "Modules/ModuleInterface.h"
 #include "Modules/ModuleManager.h"
 #include "KismetCompiledFunctionContext.h"
+#include "Engine/Blueprint.h" // for FCompilerNativizationOptions
 
 class UUserDefinedEnum;
 class UUserDefinedStruct;
diff --git a/Engine/Source/Runtime/Engine/Private/InheritableComponentHandler.cpp b/Engine/Source/Runtime/Engine/Private/InheritableComponentHandler.cpp
--- a/Engine/Source/Runtime/Engine/Private/InheritableComponentHandler.cpp
+++ b/Engine/Source/Runtime/Engine/Private/InheritableComponentHandler.cpp
@@ -7,6 +7,7 @@
 #include "UObject/PropertyPortFlags.h"
 #include "UObject/LinkerLoad.h"
 #include "UObject/BlueprintsObjectVersion.h"
+#include "UObject/UObjectHash.h" // for FindObjectWithOuter()
 
 #if WITH_EDITOR
 #include "Kismet2/BlueprintEditorUtils.h"
diff --git a/Engine/Source/Runtime/Engine/Private/UserDefinedStruct.cpp b/Engine/Source/Runtime/Engine/Private/UserDefinedStruct.cpp
--- a/Engine/Source/Runtime/Engine/Private/UserDefinedStruct.cpp
+++ b/Engine/Source/Runtime/Engine/Private/UserDefinedStruct.cpp
@@ -8,6 +8,7 @@
 #include "UObject/LinkerLoad.h"
 #include "Misc/SecureHash.h"
 #include "UObject/PropertyPortFlags.h"
+#include "Misc/PackageName.h" // for FPackageName::GetLongPackageAssetName()
 
 #if WITH_EDITOR
 #include "UserDefinedStructure/UserDefinedStructEditorData.h"
