Copying //UE4/Dev-Blueprints to //UE4/Dev-Main (Source: //UE4/Dev-Blueprints @ 3255454)

==========================
MAJOR FEATURES + CHANGES
==========================

Change 3228496 on 2016/12/09 by [REDACTED_PERSON]

	This change adds extra information to component template arrays so that the component class can be determined in builds that strip out objects of certain class types such as the editor dedicated server build.
	#Proj KismetCompiler, BlueprintGraph, UnrealEd, Core, Engine, Kismet, BlueprintCompilerCppBackend

Change 3232435 on 2016/12/13 by [REDACTED_PERSON]

	Fix for a bug introduced in CL 3228496 that caused component templates to fail to be identified by name and resulted in blueprint compilation issues for add component nodes.
	#Proj BlueprintGraph, Engine

Change 3234581 on 2016/12/14 by [REDACTED_PERSON]

	Backing out fix for UE-38842 (CL 3228496/3232435/3232564) - mapping UBlueprintGeneratedClass's ComponentTemplates array to a new format was causing issues with deferred dependency loading during serialization (trying to extract type information from a placeholder object). We're opting for a smaller/simpler solution to UE-38842, which will be to store the component information on the node itself (not with the templates).

Change 3236615 on 2016/12/15 by [REDACTED_PERSON]

	Nativization: Fixed getter fuctions in FUnconvertedWrapper, the returned ref won;t be const.

Change 3236967 on 2016/12/15 by [REDACTED_PERSON]

	Test data showing an error for jira issue UE-39808

Change 3237021 on 2016/12/15 by [REDACTED_PERSON]

	UE31622 test data

Change 3237046 on 2016/12/15 by [REDACTED_PERSON]

	UE-14123 test data

Change 3239289 on 2016/12/17 by [REDACTED_PERSON]

	[UE-38999] Dump component tree node hierarchy to the output log on error state during widget generation.

	change summary:
	- added FOnTableViewBadState delegate parameter to SSCSEditorDragDropTree
	- added SSCSEditor::DumpTree() as the FOnTableViewBadState delegate implementatioon for the STableView widget (to provide us with more info on future occurrences)

Change 3239448 on 2016/12/19 by [REDACTED_PERSON]

	New way of collecting dependencies assets. Only directly used assets are listed. It will be used in projects with EDL enabled, once the "EDL boot time" is enabled.

	Nativized projects with EDL disabled  use the new mechanism as well. They use new  __StaticDependenciesAssets functions generator. __StaticDependenciesAssets calls recursively __StaticDependenciesAssets of all Blueprints, that the current BP depends on. It reduces the size of __StaticDependenciesAssets.

	Notice, that at the moment, this change should no affect any projects with EDL enabled (because the  EDL boot time doesn't work yet).

Change 3239778 on 2016/12/19 by [REDACTED_PERSON]

	[UE-39854] Fix nativized assets build error when there are no native code dependencies.

	change summary:
	- modified FDependenciesGlobalMapHelper::EmitBodyCode() to emit a NULL entry when there are no other entries to emit, in order to avoid a zero-length array initialization error at compile time.

Change 3239965 on 2016/12/19 by [REDACTED_PERSON]

	[UE-39733] Fix incorrect graph pin value display names for user-defined enum types.

	change summary:
	- switched UEnum::GetDisplayNameText() to be a virtual API
	- added a UUserDefinedEnum::GetDisplayNameText() override to call FEnumEditorUtils::GetEnumeratorDisplayName()

Change 3240422 on 2016/12/19 by [REDACTED_PERSON]

	Remove useless counter. ensureMsgf only fires once

Change 3242313 on 2016/12/21 by [REDACTED_PERSON]

	[UE-35418] The Actor details view will now refresh property mappings for the current Actor instance selection after a bytecode-only recompile of its Blueprint class as a dependency during reinstancing.

	change summary:
	- Modified FBlueprintCompileReinstancer::ReinstanceObjects() to invoke SelectActor() on the current editor selection after a bytecode-only recompile of its class as a dependent BP (i.e. when reinstancing/finalization is not actually incurred). This is meant to be consistent with how the refresh for the current Actor selection is handled by the reinstancer in the other cases.

Change 3242409 on 2016/12/21 by [REDACTED_PERSON]

	PR #2995: Arbitrary base logarithm blueprint node (Contributed by [REDACTED_PERSON])

Change 3243207 on 2016/12/22 by [REDACTED_PERSON]

	[UE-39816] Renaming interface input/output parameters will no longer cause broken pin links at interface function call sites in Blueprints that are currently loaded.

	change summary:
	- modified FBasePinChangeHelper::Broadcast() to consider pin changes sourced from an interface Blueprint
	- also revised non-interface function call site check code a bit to try and avoid doing some unnecessary work when possible, since this code runs through every call site node

Change 3243210 on 2016/12/22 by [REDACTED_PERSON]

	[UE-39944] Extend the GetClassDefaults node to include output pin exceptions for TSet/TMap properties (i.e. mirror safeguards already in place for TArray).

	change summary:
	- deprecated 'bExcludeObjectArrays' and replaced with 'bExcludeObjectContainers' (not using a redirect because I need the old property serialized on load for backwards-compatibility)
	- modified FClassDefaultsOptionalPinManager's ctor to consider both flags
	- modified FClassDefaultsOptionalPinManager::CanTreatPropertyAsOptional() to also test for TSet/TMap types, and exclude them if they contain a non-class UObject property as the inner type (same as we do for TArray)
	- modified UK2Node_GetClassDefaults::ValidateNodeDuringCompilation() to also test for TSet/TMap types, and emit a warning for existing pin connections that would have otherwise failed for newly-placed nodes (same as TArray)
	- modified UK2Node_GetClassDefaults::ExpandNode() to generate the necessary script needed to create a copy of TSet/TMap property value outputs (same as we do for TArray types)

Change 3243373 on 2016/12/23 by [REDACTED_PERSON]

	-Nativized build with EDL enabled use new dependency-gathering system. It shrinks size of the .exe file. This code will change once the boot time EDL is enabled.
	-In nativized code, ZCOnstructors are not called. They are replaced by ::StaticClass(). Static class doesn't call noting recursively. We still need to fix interfaces (they have no StaticClass).
	- Workaround for  UE-40026
	- Various minior improvements.

Change 3244038 on 2016/12/28 by [REDACTED_PERSON]

	[UE-34488] Child Blueprint assets are now marked as dirty if its parent Blueprint incurs a strucutural change.

	change summary:
	- modified FBlueprintEditorUtils::MarkBlueprintAsStructurallyModified() to call MarkPackageDirty() after a skeleton-only compile pass on a child Blueprint.

	notes:
	- this only fixes the issue for child Blueprints that are loaded when the change to the parent Blueprint occurs. unloaded child Blueprints would still need a redirector to fix this up on load.

Change 3244087 on 2016/12/28 by [REDACTED_PERSON]

	CIS warning fix (C6236).

Change 3244388 on 2016/12/30 by [REDACTED_PERSON]

	[UE-39816] Fix broken pin links caused by renaming interface function input/output parameters prior to compiling the interface, but after renaming the function itself.

	change summary:
	- modified FBlueprintEditorUtils::FindFunctionInImplementedInterfaces() to prefer the skeleton class when searching, as that is always up-to-date.
	- modified FBlueprintEditorUtils::RenameGraph() to update *all* function call site nodes, rather than just those within the current Blueprint class scope.

Change 3245322 on 2017/01/03 by [REDACTED_PERSON]

	Fixed a crash when implementing a native interface in a BP

Change 3245667 on 2017/01/03 by [REDACTED_PERSON]

	Mirroring CL 3245664 - Making it so level script bound nodes are fixed up on load (in case the level instances were saved out without the bindings).

Change 3247675 on 2017/01/05 by [REDACTED_PERSON]

	BP is not DataOnly, when it overrides an inherited component

	related to UE-40131

Change 3247985 on 2017/01/05 by [REDACTED_PERSON]

	NativizationSummary object is always present.

Change 3249423 on 2017/01/06 by [REDACTED_PERSON]

	Mirroring CL 3248792 from Orion.

	Fix to keep placeholder classes from being needlessly created (when the object they represent already exists) - instead, attempt to lookup and find the existing import objects (which used to be set, but could be cleared during async loading by FLinkerManager::DissociateImportsAndForcedExports()).

Change 3249568 on 2017/01/06 by [REDACTED_PERSON]

	Updating the UBlueprintThumbnailRenderer API so it can be sub-classed by plugin/projects.

	PR #2899: Expose UBlueprintThumbnailRenderer in the API (Contributed by [REDACTED_PERSON])

Change 3251903 on 2017/01/10 by [REDACTED_PERSON]

	[UE-31640] Function inputs are now exposed as variable "Get" nodes via the right-click context menu in a Blueprint function graph context.

	Note: Changes are based on shelved CL# 2452724, original proof-of-concept code (credit: Michael.Noland).

Change 3252119 on 2017/01/10 by [REDACTED_PERSON]

	CIS fix (shadowed variable).

Change 3252744 on 2017/01/10 by [REDACTED_PERSON]

	Add GMinimalCompileOnLoad path for postponing compile on load until all data has been loaded - long term this will be used to improve bp compiler infrastructure (performance, correctness, succinctness)

Change 3252968 on 2017/01/10 by [REDACTED_PERSON]

	[UE-36798] Fix for an infinite loop case in the math expression parser.

Change 3253153 on 2017/01/10 by [REDACTED_PERSON]

	Fixed construction scripts not being loaded before instancing occurs when using new lighter compile path

Change 3253171 on 2017/01/10 by [REDACTED_PERSON]

	Mirrored CL 3253147.

	Properly fills out FPropertyChangedEvent's MemberProperty field, for the property editor's NotifyHook. As its comment implies, "MemberProperty" is meant to represent the outermost property (in scenarios, like with nested struct properties). It was not working this way, and was instead set to the same nested property.

Change 3253220 on 2017/01/10 by [REDACTED_PERSON]

	These pins should infer together

Change 3253223 on 2017/01/10 by [REDACTED_PERSON]

	[UE-35050] Fix a crash that occurs on an attempt to rename the category of an implemented interface function inherited from a native parent class in the My Blueprint panel.

	change summary:
	- modified FBlueprintEditorUtils::FindFunctionInImplementedInterfaces() to accept an additional input parameter so that callers can opt-in to searching all inherited interface classes
	- modified SMyBlueprint::CanRequestRenameOnActionNode() to prevent rename on implemented interface functions inherited from a native parent class

Change 3253259 on 2017/01/10 by [REDACTED_PERSON]

	Quick cleanup, this stuff is in progress

Change 3253983 on 2017/01/11 by [REDACTED_PERSON]

	[UE-35629] Disable external curve asset creation when editing a local variable's default value in a Blueprint graph.

	change summary:
	- modified FCurveStructCustomization::CustomizeChildren() to set the "create" button visibility to "collapsed" when the Owner is NULL (as the default asset path is inferred from the Owner's package).

Change 3254024 on 2017/01/11 by [REDACTED_PERSON]

	[UE-40131] Non-native child BPs can now properly override a nativized parent BP's components in a cooked build with exclusive Blueprint class nativiation.

	change summary:
	- added UBlueprintGeneratedClass::CheckAndApplyComponentTemplateOverrides()
	- switched FComponentKey::OwnerClass type from UBlueprintGeneratedClass to UClass to allow the reference to be loaded when it's a nativized BP class.
	- modified AActor::PostLoadSubobjects() to apply ICH template overrides to inherited component subobjects on load (only in a cooked build)
	- modified AActor::PostSpawnInitialize() to also apply ICH template overrides to inherited component subobjects on new Actor spawns (only in a cooked build)
	- modified FBlueprintEditorUtils::IsDataOnlyBlueprint() to additionally reject BPs that have a non-empty ICH (from Maciej)
	- added UBlueprint::bHasNativizedParent to cache this as an indicator for cooked builds to allow Actor instances to avoid incurring additional overhead for this at load time in the general case

Change 3254176 on 2017/01/11 by [REDACTED_PERSON]

	Mirroring CL 3245838 to unblock the Odin build. Downgrading new error to warning until we can address it in content (and fully understand why it is triggering).

Change 3254391 on 2017/01/11 by [REDACTED_PERSON]

	[UE-40131] CIS fix (non-unity).

Change 3254442 on 2017/01/11 by [REDACTED_PERSON]

	Mirroring CL 3245069 (from RobM), which is described as a "temp fix".

Change 3254599 on 2017/01/11 by [REDACTED_PERSON]

	Updating some of our GetRedirectPinNames() functions - making local copy of an array element instead of a reference. We're adding to the array right after, and the ref can become invalid (if the array is grown and reallocated). Came to us via UDN.

Change 3254624 on 2017/01/11 by [REDACTED_PERSON]

	Backing out CL 3247675, as it removed the "data-only" status from certain Blueprints that, IMO, are still considered "data-only". This would have reprecussions to editor (compile-on) load times, and affect how these Blueprints are presented to the user.

Change 3254671 on 2017/01/11 by [REDACTED_PERSON]

	Resolving CIS warning - USE_EVENT_DRIVEN_ASYNC_LOAD was replaced with GEventDrivenLoaderEnabled.

diff --git a/Engine/Source/Developer/BlueprintCompilerCppBackend/Private/BlueprintCompilerCppBackendBase.cpp b/Engine/Source/Developer/BlueprintCompilerCppBackend/Private/BlueprintCompilerCppBackendBase.cpp
--- a/Engine/Source/Developer/BlueprintCompilerCppBackend/Private/BlueprintCompilerCppBackendBase.cpp
+++ b/Engine/Source/Developer/BlueprintCompilerCppBackend/Private/BlueprintCompilerCppBackendBase.cpp
@@ -1057,7 +1057,7 @@ FString FBlueprintCompilerCppBackendBase::GenerateWrapperForClass(UClass* Source
 			? FString::Printf(TEXT("%s::%s"), *DelegatesClassName, *GenerateMulticastDelegateTypeName(CastChecked<UMulticastDelegateProperty>(Property)))
 			: EmitterContext.ExportCppDeclaration(Property
 				, EExportedDeclaration::Parameter
-				, EPropertyExportCPPFlags::CPPF_CustomTypeName | EPropertyExportCPPFlags::CPPF_BlueprintCppBackend | EPropertyExportCPPFlags::CPPF_NoRef
+				, EPropertyExportCPPFlags::CPPF_CustomTypeName | EPropertyExportCPPFlags::CPPF_BlueprintCppBackend | EPropertyExportCPPFlags::CPPF_NoRef | EPropertyExportCPPFlags::CPPF_NoConst
 				, FEmitterLocalContext::EPropertyNameInDeclaration::Skip);
 		EmitterContext.Header.AddLine(FString::Printf(TEXT("FORCENOINLINE %s& GetRef__%s()"), *TypeDeclaration, *UnicodeToCPPIdentifier(Property->GetName(), false, nullptr)));
 		EmitterContext.Header.AddLine(TEXT("{"));
diff --git a/Engine/Source/Developer/BlueprintCompilerCppBackend/Private/BlueprintCompilerCppBackendGatherDependencies.cpp b/Engine/Source/Developer/BlueprintCompilerCppBackend/Private/BlueprintCompilerCppBackendGatherDependencies.cpp
--- a/Engine/Source/Developer/BlueprintCompilerCppBackend/Private/BlueprintCompilerCppBackendGatherDependencies.cpp
+++ b/Engine/Source/Developer/BlueprintCompilerCppBackend/Private/BlueprintCompilerCppBackendGatherDependencies.cpp
@@ -116,7 +116,9 @@ struct FFindAssetsToInclude : public FGatherConvertedClassDependenciesHelperBase
 
 		const bool bUseZConstructorInGeneratedCode = false;
 		//TODO: What About Delegates?
-		auto ObjAsBPGC = Cast<UBlueprintGeneratedClass>(Object);
+		UField* AsField = Cast<UField>(Object);
+		UBlueprintGeneratedClass* ObjAsBPGC = Cast<UBlueprintGeneratedClass>(Object);
+		//TODO: update once the BootTimeEDL is ready
 		const bool bWillBeConvetedAsBPGC = ObjAsBPGC && Dependencies.WillClassBeConverted(ObjAsBPGC);
 		if (bWillBeConvetedAsBPGC)
 		{
@@ -128,6 +130,7 @@ struct FFindAssetsToInclude : public FGatherConvertedClassDependenciesHelperBase
 					IncludeTheHeaderInBody(ObjAsBPGC);
 				}
 			}
+			return;
 		}
 		else if (UUserDefinedStruct* UDS = Cast<UUserDefinedStruct>(Object))
 		{
@@ -147,28 +150,28 @@ struct FFindAssetsToInclude : public FGatherConvertedClassDependenciesHelperBase
 				Dependencies.ConvertedEnum.Add(UDE);
 			}
 		}
-		else if ((Object->IsAsset() || ObjAsBPGC) && !Object->IsIn(CurrentlyConvertedStruct))
+		else if ((Object->IsAsset() || AsField) && !Object->IsIn(CurrentlyConvertedStruct))
 		{
-			// include all not converted super classes
-			for (auto SuperBPGC = ObjAsBPGC ? Cast<UBlueprintGeneratedClass>(ObjAsBPGC->GetSuperClass()) : nullptr;
-				SuperBPGC && !Dependencies.WillClassBeConverted(SuperBPGC);
-				SuperBPGC = Cast<UBlueprintGeneratedClass>(SuperBPGC->GetSuperClass()))
+			if (AsField)
 			{
-				Dependencies.Assets.AddUnique(SuperBPGC);
+				UClass* OwnerClass = AsField->GetOwnerClass();
+				if (OwnerClass)
+				{
+					Dependencies.Assets.AddUnique(OwnerClass);
+				}
+				else if (UStruct* OwnerStruct = AsField->GetOwnerStruct())
+				{
+					Dependencies.Assets.AddUnique(OwnerStruct);
+				}
+				else
+				{
+					Dependencies.Assets.AddUnique(Object);
+				}
 			}
-
-			Dependencies.Assets.AddUnique(Object);
-			return;
-		}
-		else if (auto ObjAsClass = Cast<UClass>(Object))
-		{
-			if (ObjAsClass->HasAnyClassFlags(CLASS_Native))
+			else
 			{
-				return;
+				Dependencies.Assets.AddUnique(Object);
 			}
-		}
-		else if (Object->IsA<UScriptStruct>())
-		{
 			return;
 		}
 
@@ -550,7 +553,7 @@ TSet<const UObject*> FGatherConvertedClassDependencies::AllDependencies() const
 	TSet<const UObject*> All;
 
 	UBlueprintGeneratedClass* SuperClass = Cast<UBlueprintGeneratedClass>(OriginalStruct->GetSuperStruct());
-	if (SuperClass && WillClassBeConverted(SuperClass))
+	if (OriginalStruct->GetSuperStruct() && (!SuperClass || WillClassBeConverted(SuperClass)))
 	{
 		All.Add(SuperClass);
 	}
@@ -560,7 +563,7 @@ TSet<const UObject*> FGatherConvertedClassDependencies::AllDependencies() const
 		for (auto& ImplementedInterface : SourceClass->Interfaces)
 		{
 			UBlueprintGeneratedClass* InterfaceClass = Cast<UBlueprintGeneratedClass>(ImplementedInterface.Class);
-			if (InterfaceClass && WillClassBeConverted(InterfaceClass))
+			if (ImplementedInterface.Class && (!InterfaceClass || WillClassBeConverted(InterfaceClass)))
 			{
 				All.Add(InterfaceClass);
 			}
diff --git a/Engine/Source/Developer/BlueprintCompilerCppBackend/Private/BlueprintCompilerCppBackendUtils.cpp b/Engine/Source/Developer/BlueprintCompilerCppBackend/Private/BlueprintCompilerCppBackendUtils.cpp
--- a/Engine/Source/Developer/BlueprintCompilerCppBackend/Private/BlueprintCompilerCppBackendUtils.cpp
+++ b/Engine/Source/Developer/BlueprintCompilerCppBackend/Private/BlueprintCompilerCppBackendUtils.cpp
@@ -1873,18 +1873,25 @@ FString FDependenciesGlobalMapHelper::EmitBodyCode()
 			DependenciesGlobalMap.GenerateValueArray(DependenciesArray);
 		}
 
-		DependenciesArray.Sort(
-			[](const FNativizationSummary::FDependencyRecord& A, const FNativizationSummary::FDependencyRecord& B) -> bool
-		{
-			return A.Index < B.Index;
-		});
-		int32 Index = 0;
-		for (FNativizationSummary::FDependencyRecord& Record : DependenciesArray)
-		{
-			ensure(!Record.NativeLine.IsEmpty());
-			ensure(Record.Index == Index);
-			Index++;
-			CodeText.AddLine(Record.NativeLine);
+		if (DependenciesArray.Num() > 0)
+		{
+			DependenciesArray.Sort(
+				[](const FNativizationSummary::FDependencyRecord& A, const FNativizationSummary::FDependencyRecord& B) -> bool
+			{
+				return A.Index < B.Index;
+			});
+			int32 Index = 0;
+			for (FNativizationSummary::FDependencyRecord& Record : DependenciesArray)
+			{
+				ensure(!Record.NativeLine.IsEmpty());
+				ensure(Record.Index == Index);
+				Index++;
+				CodeText.AddLine(Record.NativeLine);
+			}
+		}
+		else
+		{
+			CodeText.AddLine(TEXT("FBlueprintDependencyObjectRef()"));
 		}
 
 		CodeText.AddLine(TEXT("};"));
diff --git a/Engine/Source/Developer/BlueprintCompilerCppBackend/Private/BlueprintCompilerCppBackendUtils.h b/Engine/Source/Developer/BlueprintCompilerCppBackend/Private/BlueprintCompilerCppBackendUtils.h
--- a/Engine/Source/Developer/BlueprintCompilerCppBackend/Private/BlueprintCompilerCppBackendUtils.h
+++ b/Engine/Source/Developer/BlueprintCompilerCppBackend/Private/BlueprintCompilerCppBackendUtils.h
@@ -348,7 +348,8 @@ private:
 	static FString HandleSpecialTypes(FEmitterLocalContext& Context, const UProperty* Property, const uint8* ValuePtr);
 
 	static FString HandleNonNativeComponent(FEmitterLocalContext& Context, const USCS_Node* Node, TSet<const UProperty*>& OutHandledProperties
-		, TArray<FString>& NativeCreatedComponentProperties, const USCS_Node* ParentNode, TArray<FNonativeComponentData>& ComponentsToInit);
+		, TArray<FString>& NativeCreatedComponentProperties, const USCS_Node* ParentNode, TArray<FNonativeComponentData>& ComponentsToInit
+		, bool bBlockRecursion);
 
 	static FString HandleInstancedSubobject(FEmitterLocalContext& Context, UObject* Object, bool bCreateInstance = true, bool bSkipEditorOnlyCheck = false);
 };
diff --git a/Engine/Source/Developer/BlueprintCompilerCppBackend/Private/BlueprintCompilerCppBackendValueHelper.cpp b/Engine/Source/Developer/BlueprintCompilerCppBackend/Private/BlueprintCompilerCppBackendValueHelper.cpp
--- a/Engine/Source/Developer/BlueprintCompilerCppBackend/Private/BlueprintCompilerCppBackendValueHelper.cpp
+++ b/Engine/Source/Developer/BlueprintCompilerCppBackend/Private/BlueprintCompilerCppBackendValueHelper.cpp
@@ -732,6 +732,7 @@ FString FEmitDefaultValueHelper::HandleSpecialTypes(FEmitterLocalContext& Contex
 
 struct FNonativeComponentData
 {
+	const USCS_Node* SCSNode;
 	FString NativeVariablePropertyName;
 	UActorComponent* ComponentTemplate;
 	UObject* ObjectToCompare;
@@ -744,7 +745,8 @@ struct FNonativeComponentData
 	bool bIsRoot;
 
 	FNonativeComponentData()
-		: ComponentTemplate(nullptr)
+		: SCSNode(nullptr)
+		, ComponentTemplate(nullptr)
 		, ObjectToCompare(nullptr)
 		, bSetNativeCreationMethod(false)
 		, bIsRoot(false)
@@ -839,7 +841,10 @@ struct FNonativeComponentData
 	}
 };
 
-FString FEmitDefaultValueHelper::HandleNonNativeComponent(FEmitterLocalContext& Context, const USCS_Node* Node, TSet<const UProperty*>& OutHandledProperties, TArray<FString>& NativeCreatedComponentProperties, const USCS_Node* ParentNode, TArray<FNonativeComponentData>& ComponenntsToInit)
+FString FEmitDefaultValueHelper::HandleNonNativeComponent(FEmitterLocalContext& Context, const USCS_Node* Node
+	, TSet<const UProperty*>& OutHandledProperties, TArray<FString>& NativeCreatedComponentProperties
+	, const USCS_Node* ParentNode, TArray<FNonativeComponentData>& ComponenntsToInit
+	, bool bBlockRecursion)
 {
 	check(Node);
 	check(Context.CurrentCodeType == FEmitterLocalContext::EGeneratedCodeType::CommonConstructor);
@@ -861,11 +866,18 @@ FString FEmitDefaultValueHelper::HandleNonNativeComponent(FEmitterLocalContext&
 			NativeVariablePropertyName = VariableCleanName;
 		}
 
+		//TODO: UGLY HACK UE-40026
+		if (bBlockRecursion && Context.CommonSubobjectsMap.Contains(ComponentTemplate))
+		{
+			return FString();
+		}
+
 		Context.AddCommonSubObject_InConstructor(ComponentTemplate, NativeVariablePropertyName);
 
 		if (ComponentTemplate->GetOuter() == BPGC)
 		{
 			FNonativeComponentData NonativeComponentData;
+			NonativeComponentData.SCSNode = Node;
 			NonativeComponentData.NativeVariablePropertyName = NativeVariablePropertyName;
 			NonativeComponentData.ComponentTemplate = ComponentTemplate;
 			USCS_Node* RootComponentNode = nullptr;
@@ -911,9 +923,12 @@ FString FEmitDefaultValueHelper::HandleNonNativeComponent(FEmitterLocalContext&
 	}
 
 	// Recursively handle child nodes.
-	for (auto ChildNode : Node->ChildNodes)
+	if (!bBlockRecursion)
 	{
-		HandleNonNativeComponent(Context, ChildNode, OutHandledProperties, NativeCreatedComponentProperties, Node, ComponenntsToInit);
+		for (auto ChildNode : Node->ChildNodes)
+		{
+			HandleNonNativeComponent(Context, ChildNode, OutHandledProperties, NativeCreatedComponentProperties, Node, ComponenntsToInit, bBlockRecursion);
+		}
 	}
 
 	return NativeVariablePropertyName;
@@ -1183,7 +1198,7 @@ void FEmitDefaultValueHelper::AddStaticFunctionsForDependencies(FEmitterLocalCon
 		FakeImportTableHelper.FillDependencyData(InAsset, Result);
 		return Result;
 	};
-
+	const bool bBootTimeEDL = false;
 	auto AddAssetArray = [&](const TArray<const UObject*>& Assets)
 	{
 		if (Assets.Num())
@@ -1220,15 +1235,40 @@ void FEmitDefaultValueHelper::AddStaticFunctionsForDependencies(FEmitterLocalCon
 			Context.AddLine(TEXT("};"));
 			Context.AddLine(TEXT("for(const FCompactBlueprintDependencyData CompactData : LocCompactBlueprintDependencyData)"));
 			Context.AddLine(TEXT("{"));
-			Context.AddLine(TEXT("\tAssetsToLoad.Add(FBlueprintDependencyData("));
+			if (IsEventDrivenLoaderEnabledInCookedBuilds() && bBootTimeEDL)
+			{
+				Context.AddLine(TEXT("\tAssetsToLoad.Add(FBlueprintDependencyData("));
+			}
+			else
+			{
+				Context.AddLine(TEXT("\tAssetsToLoad.AddUnique(FBlueprintDependencyData("));
+			}
 			Context.AddLine(TEXT("\t\tF__NativeDependencies::Get(CompactData.ObjectRefIndex)"));
 			Context.AddLine(TEXT("\t\t,CompactData.ClassDependency"));
 			Context.AddLine(TEXT("\t\t,CompactData.CDODependency"));
+			Context.AddLine(TEXT("\t\t,CompactData.ObjectRefIndex"));
 			Context.AddLine(TEXT("\t));"));
 			Context.AddLine(TEXT("} "));
 		}
 	};
 
+	TSet<const UBlueprintGeneratedClass*> OtherBPGCs;
+	if (!IsEventDrivenLoaderEnabledInCookedBuilds() & !bBootTimeEDL)
+	{
+		for (const UObject* It : AllDependenciesToHandle)
+		{
+			if (const UBlueprintGeneratedClass* OtherBPGC = Cast<const UBlueprintGeneratedClass>(It))
+			{
+				const UBlueprint* BP = Cast<const UBlueprint>(OtherBPGC->ClassGeneratedBy);
+				if (Context.Dependencies.WillClassBeConverted(OtherBPGC) && BP && (BP->BlueprintType != EBlueprintType::BPTYPE_Interface))
+				{
+					OtherBPGCs.Add(OtherBPGC);
+				}
+			}
+		}
+	}
+
+
 	// 3. LIST OF UsedAssets
 	{
 		Context.AddLine(FString::Printf(TEXT("void %s::__StaticDependencies_DirectlyUsedAssets(TArray<FBlueprintDependencyData>& AssetsToLoad)"), *CppClassName));
@@ -1252,7 +1292,72 @@ void FEmitDefaultValueHelper::AddStaticFunctionsForDependencies(FEmitterLocalCon
 		Context.AddLine(FString::Printf(TEXT("void %s::__StaticDependenciesAssets(TArray<FBlueprintDependencyData>& AssetsToLoad)"), *CppClassName));
 		Context.AddLine(TEXT("{"));
 		Context.IncreaseIndent();
-		Context.AddLine(FString(TEXT("__StaticDependencies_DirectlyUsedAssets(AssetsToLoad);")));
+
+		if (!OtherBPGCs.Num() || (IsEventDrivenLoaderEnabledInCookedBuilds() && bBootTimeEDL))
+		{
+			Context.AddLine(FString(TEXT("__StaticDependencies_DirectlyUsedAssets(AssetsToLoad);")));
+		}
+		else
+		{
+			// To reduce the size of __StaticDependenciesAssets, all __StaticDependenciesAssets of listed BPs will be called.
+			FNativizationSummary::FDependencyRecord& DependencyRecord = FDependenciesGlobalMapHelper::FindDependencyRecord(OriginalClass);
+			ensure(DependencyRecord.Index >= 0); 
+			if (DependencyRecord.NativeLine.IsEmpty())
+			{
+				DependencyRecord.NativeLine = CreateAssetToLoadString(OriginalClass);
+			}
+			Context.AddLine(FString::Printf(TEXT("const int16 __OwnIndex = %d;"), DependencyRecord.Index));
+			Context.AddLine(FString(TEXT("if(FBlueprintDependencyData::ContainsDependencyData(AssetsToLoad, __OwnIndex)) { return; }")));
+			Context.AddLine(TEXT("if(GEventDrivenLoaderEnabled && EVENT_DRIVEN_ASYNC_LOAD_ACTIVE_AT_RUNTIME){ __StaticDependencies_DirectlyUsedAssets(AssetsToLoad); }"));
+			Context.AddLine(TEXT("else"));
+			Context.AddLine(TEXT("{"));
+			Context.IncreaseIndent();
+			Context.AddLine(FString(TEXT("const bool __FirstFunctionCall = !AssetsToLoad.Num();")));
+			Context.AddLine(FString(TEXT("TArray<FBlueprintDependencyData> Temp;")));
+			// Other __StaticDependenciesAssets fucntions should not see the assets added by __StaticDependencies_DirectlyUsedAssets
+			// But in the first function called the assets from __StaticDependencies_DirectlyUsedAssets must go first in unchanged order (to satisfy FConvertedBlueprintsDependencies::FillUsedAssetsInDynamicClass)
+			Context.AddLine(FString(TEXT("__StaticDependencies_DirectlyUsedAssets(__FirstFunctionCall ? AssetsToLoad : Temp);")));
+			Context.AddLine(FString(TEXT("TArray<FBlueprintDependencyData>& ArrayUnaffectedByDirectlyUsedAssets = __FirstFunctionCall ? Temp : AssetsToLoad;")));
+
+			Context.AddLine(FString(TEXT("ArrayUnaffectedByDirectlyUsedAssets.AddUnique(FBlueprintDependencyData(F__NativeDependencies::Get(__OwnIndex), {}, {}, __OwnIndex));")));
+
+			for (const UBlueprintGeneratedClass* OtherBPGC : OtherBPGCs)
+			{
+				Context.AddLine(FString::Printf(TEXT("%s::__StaticDependenciesAssets(ArrayUnaffectedByDirectlyUsedAssets);"), *FEmitHelper::GetCppName(OtherBPGC)));
+			}
+			Context.AddLine(FString(TEXT("FBlueprintDependencyData::AppendUniquely(AssetsToLoad, Temp);")));
+			Context.DecreaseIndent();
+			Context.AddLine(TEXT("}"));
+		}
+
+		if (IsEventDrivenLoaderEnabledInCookedBuilds() && bBootTimeEDL)
+		{
+			//TODO: remove stuff from CoreUObject
+		}
+		else
+		{
+			//WIthout EDL we don't need the native stuff.
+			for (auto Iter = AllDependenciesToHandle.CreateIterator(); Iter; ++Iter)
+			{
+				const UObject* ItObj = *Iter;
+				if (auto ObjAsClass = Cast<const UClass>(ItObj))
+				{
+					if (ObjAsClass->HasAnyClassFlags(CLASS_Native))
+					{
+						Iter.RemoveCurrent();
+					}
+				}
+				else if (ItObj && ItObj->IsA<UScriptStruct>() && !ItObj->IsA<UUserDefinedStruct>())
+				{
+					Iter.RemoveCurrent();
+				}
+				else if (ItObj && ItObj->IsA<UEnum>() && !ItObj->IsA<UUserDefinedEnum>())
+				{
+					Iter.RemoveCurrent();
+				}
+			}
+		}
+		
 		AddAssetArray(AllDependenciesToHandle.Array());
 		Context.DecreaseIndent();
 		Context.AddLine(TEXT("}"));
@@ -1333,8 +1438,17 @@ void FEmitDefaultValueHelper::GenerateCustomDynamicClassInitialization(FEmitterL
 				continue;
 			}
 
-			const FString ClassConstructor = FDependenciesHelper::GenerateZConstructor(ClassToLoad);
-			Context.AddLine(FString::Printf(TEXT("extern UClass* %s;"), *ClassConstructor));
+			FString ClassConstructor;
+			if (ClassToLoad->HasAnyClassFlags(CLASS_Interface))
+			{
+				const FString ClassZConstructor = FDependenciesHelper::GenerateZConstructor(ClassToLoad);
+				Context.AddLine(FString::Printf(TEXT("extern UClass* %s;"), *ClassZConstructor));
+				ClassConstructor = ClassZConstructor;
+			}
+			else
+			{
+				ClassConstructor = FString::Printf(TEXT("%s::StaticClass()"), *FEmitHelper::GetCppName(ClassToLoad));
+			}
 			Context.AddLine(FString::Printf(TEXT("InDynamicClass->%s.Add(%s);"), GET_MEMBER_NAME_STRING_CHECKED(UDynamicClass, ReferencedConvertedFields), *ClassConstructor));
 
 			//Context.AddLine(FString::Printf(TEXT("InDynamicClass->ReferencedConvertedFields.Add(LoadObject<UClass>(nullptr, TEXT(\"%s\")));")
@@ -1530,11 +1644,11 @@ void FEmitDefaultValueHelper::GenerateConstructor(FEmitterLocalContext& Context)
 			{
 				if (BPGCStack[i]->SimpleConstructionScript)
 				{
-					for (auto Node : BPGCStack[i]->SimpleConstructionScript->GetRootNodes())
+					for (USCS_Node* Node : BPGCStack[i]->SimpleConstructionScript->GetRootNodes())
 					{
 						if (Node)
 						{
-							const FString NativeVariablePropertyName = HandleNonNativeComponent(Context, Node, HandledProperties, NativeCreatedComponentProperties, nullptr, ComponentsToInit);
+							const FString NativeVariablePropertyName = HandleNonNativeComponent(Context, Node, HandledProperties, NativeCreatedComponentProperties, nullptr, ComponentsToInit, false);
 
 							if (bNeedsRootComponentAssignment && Node->ComponentTemplate && Node->ComponentTemplate->IsA<USceneComponent>() && !NativeVariablePropertyName.IsEmpty())
 							{
@@ -1550,6 +1664,20 @@ void FEmitDefaultValueHelper::GenerateConstructor(FEmitterLocalContext& Context)
 							}
 						}
 					}
+
+					//TODO: UGLY HACK for "zombie" nodes - UE-40026
+					for (USCS_Node* Node : BPGCStack[i]->SimpleConstructionScript->GetAllNodes())
+					{
+						if (Node)
+						{
+							const bool bNodeWasProcessed = nullptr != ComponentsToInit.FindByPredicate([=](const FNonativeComponentData& InData) { return Node == InData.SCSNode; });
+							if (!bNodeWasProcessed)
+							{
+								HandleNonNativeComponent(Context, Node, HandledProperties, NativeCreatedComponentProperties, nullptr, ComponentsToInit, true);
+							}
+						}
+					}
+
 				}
 			}
 
diff --git a/Engine/Source/Developer/BlueprintNativeCodeGen/Private/BlueprintNativeCodeGenModule.cpp b/Engine/Source/Developer/BlueprintNativeCodeGen/Private/BlueprintNativeCodeGenModule.cpp
--- a/Engine/Source/Developer/BlueprintNativeCodeGen/Private/BlueprintNativeCodeGenModule.cpp
+++ b/Engine/Source/Developer/BlueprintNativeCodeGen/Private/BlueprintNativeCodeGenModule.cpp
@@ -323,6 +323,11 @@ void FBlueprintNativeCodeGenModule::InitializeForRerunDebugOnly(const TArray<FPl
 void FBlueprintNativeCodeGenModule::GenerateFullyConvertedClasses()
 {
 	TSharedPtr<FNativizationSummary> NativizationSummary(new FNativizationSummary());
+	{
+		IBlueprintCompilerCppBackendModule& CodeGenBackend = (IBlueprintCompilerCppBackendModule&)IBlueprintCompilerCppBackendModule::Get();
+		CodeGenBackend.NativizationSummary() = NativizationSummary;
+	}
+
 	for (TAssetPtr<UBlueprint>& BPPtr : ToGenerate)
 	{
 		UBlueprint* BP = BPPtr.LoadSynchronous();
diff --git a/Engine/Source/Editor/BlueprintGraph/Classes/BlueprintVariableNodeSpawner.h b/Engine/Source/Editor/BlueprintGraph/Classes/BlueprintVariableNodeSpawner.h
--- a/Engine/Source/Editor/BlueprintGraph/Classes/BlueprintVariableNodeSpawner.h
+++ b/Engine/Source/Editor/BlueprintGraph/Classes/BlueprintVariableNodeSpawner.h
@@ -36,10 +36,11 @@ public:
 	 * 
 	 * @param  NodeClass	The node type that you want the spawner to spawn.
 	 * @param  VarProperty	The property that represents the member-variable you want nodes spawned for.
+	 * @param  VarContext	The graph that the local variable belongs to.
 	 * @param  Outer		Optional outer for the new spawner (if left null, the transient package will be used).
 	 * @return A newly allocated instance of this class.
 	 */
-	static UBlueprintVariableNodeSpawner* Create(TSubclassOf<UK2Node_Variable> NodeClass, UProperty const* VarProperty, UObject* Outer = nullptr);
+	static UBlueprintVariableNodeSpawner* CreateFromMemberOrParam(TSubclassOf<UK2Node_Variable> NodeClass, UProperty const* VarProperty, UEdGraph* VarContext = nullptr, UObject* Outer = nullptr);
 
 	/**
 	 * Creates a new UBlueprintVariableNodeSpawner, charged with spawning
@@ -52,7 +53,7 @@ public:
 	 * @param  Outer		Optional outer for the new spawner (if left null, the transient package will be used).
 	 * @return A newly allocated instance of this class.
 	 */
-	static UBlueprintVariableNodeSpawner* Create(TSubclassOf<UK2Node_Variable> NodeClass, UEdGraph* VarContext, FBPVariableDescription const& VarDesc, UProperty* VarProperty, UObject* Outer = nullptr);
+	static UBlueprintVariableNodeSpawner* CreateFromLocal(TSubclassOf<UK2Node_Variable> NodeClass, UEdGraph* VarContext, FBPVariableDescription const& VarDesc, UProperty* VarProperty, UObject* Outer = nullptr);
 
 	// UBlueprintNodeSpawner interface
 	virtual void Prime() override;
@@ -61,6 +62,11 @@ public:
 	virtual UEdGraphNode* Invoke(UEdGraph* ParentGraph, FBindingSet const& Bindings, FVector2D const Location) const override;
 	// End UBlueprintNodeSpawner interface
 	
+	/**
+	 * @return True if this is a user-created local variable
+	 */
+	bool IsUserLocalVariable() const;
+
 	/**
 	 * Since this spawner can wrap both local and member variables, we use this
 	 * method to discern between the two.
@@ -112,4 +118,4 @@ private:
 	/** */
 	UPROPERTY()
 	FBPVariableDescription LocalVarDesc;
-};
+};
\ No newline at end of file
diff --git a/Engine/Source/Editor/BlueprintGraph/Classes/K2Node_GetClassDefaults.h b/Engine/Source/Editor/BlueprintGraph/Classes/K2Node_GetClassDefaults.h
--- a/Engine/Source/Editor/BlueprintGraph/Classes/K2Node_GetClassDefaults.h
+++ b/Engine/Source/Editor/BlueprintGraph/Classes/K2Node_GetClassDefaults.h
@@ -99,7 +99,11 @@ private:
 	UPROPERTY(EditAnywhere, Category=PinOptions, EditFixedSize)
 	TArray<FOptionalPinFromProperty> ShowPinForProperties;
 
-	/** Whether or not to exclude object array properties */
+	/** Whether or not to exclude object container properties */
 	UPROPERTY()
-	bool bExcludeObjectArrays;
+	bool bExcludeObjectContainers;
+
+	/** Whether or not to exclude object array properties (deprecated) */
+	UPROPERTY()
+	bool bExcludeObjectArrays_DEPRECATED;
 };
diff --git a/Engine/Source/Editor/BlueprintGraph/Private/BasicTokenParser.cpp b/Engine/Source/Editor/BlueprintGraph/Private/BasicTokenParser.cpp
--- a/Engine/Source/Editor/BlueprintGraph/Private/BasicTokenParser.cpp
+++ b/Engine/Source/Editor/BlueprintGraph/Private/BasicTokenParser.cpp
@@ -807,7 +807,10 @@ Loop:
 			ClearCachedComment();
 			SetError(FErrorState::ParseError, LOCTEXT("NoCommentClose", "No end to a comment by the end of the expression"));
 		}
-		goto Loop;
+		else
+		{
+			goto Loop;
+		}
 	}
 	return c;
 }
diff --git a/Engine/Source/Editor/BlueprintGraph/Private/BlueprintActionDatabase.cpp b/Engine/Source/Editor/BlueprintGraph/Private/BlueprintActionDatabase.cpp
--- a/Engine/Source/Editor/BlueprintGraph/Private/BlueprintActionDatabase.cpp
+++ b/Engine/Source/Editor/BlueprintGraph/Private/BlueprintActionDatabase.cpp
@@ -687,9 +687,9 @@ static void BlueprintActionDatabaseImpl::AddClassPropertyActions(UClass const* c
  		}
 		else
 		{
-			UBlueprintVariableNodeSpawner* GetterSpawner = UBlueprintVariableNodeSpawner::Create(UK2Node_VariableGet::StaticClass(), Property);
+			UBlueprintVariableNodeSpawner* GetterSpawner = UBlueprintVariableNodeSpawner::CreateFromMemberOrParam(UK2Node_VariableGet::StaticClass(), Property);
 			ActionListOut.Add(GetterSpawner);
-			UBlueprintVariableNodeSpawner* SetterSpawner = UBlueprintVariableNodeSpawner::Create(UK2Node_VariableSet::StaticClass(), Property);
+			UBlueprintVariableNodeSpawner* SetterSpawner = UBlueprintVariableNodeSpawner::CreateFromMemberOrParam(UK2Node_VariableSet::StaticClass(), Property);
 			ActionListOut.Add(SetterSpawner);
 		}
 	}
@@ -745,7 +745,7 @@ static void BlueprintActionDatabaseImpl::AddBlueprintGraphActions(UBlueprint con
 		ActionListOut.Add(MakeMacroNodeSpawner(*GraphIt));
 	}
 
-	// local variables
+	// local variables and parameters
 	for (auto GraphIt = Blueprint->FunctionGraphs.CreateConstIterator(); GraphIt; ++GraphIt)
 	{
 		UEdGraph* FunctionGraph = (*GraphIt);
@@ -755,15 +755,38 @@ static void BlueprintActionDatabaseImpl::AddBlueprintGraphActions(UBlueprint con
 
 		for (UK2Node_FunctionEntry* FunctionEntry : GraphEntryNodes)
 		{
+			// Find the initial place where the function was defined, so we use the most generous scope
+			UFunction* SkeletonFunction = FindField<UFunction>(Blueprint->SkeletonGeneratedClass, FunctionGraph->GetFName());
+			for (UFunction* ParentFunction = SkeletonFunction; ParentFunction != nullptr; ParentFunction = ParentFunction->GetSuperFunction())
+			{
+				SkeletonFunction = ParentFunction;
+			}
+
+			// Create entries for function parameters
+			if (SkeletonFunction != nullptr)
+			{
+				for (TFieldIterator<UProperty> ParamIt(SkeletonFunction); ParamIt && (ParamIt->PropertyFlags & CPF_Parm); ++ParamIt)
+				{
+					UProperty* Param = *ParamIt;
+					const bool bIsFunctionInput = !Param->HasAnyPropertyFlags(CPF_ReturnParm) && (!Param->HasAnyPropertyFlags(CPF_OutParm) || Param->HasAnyPropertyFlags(CPF_ReferenceParm));
+					if (bIsFunctionInput)
+					{
+						UBlueprintNodeSpawner* GetVarSpawner = UBlueprintVariableNodeSpawner::CreateFromMemberOrParam(UK2Node_VariableGet::StaticClass(), Param, FunctionGraph);
+						ActionListOut.Add(GetVarSpawner);
+					}
+				}
+			}
+
+			// Create entries for local variables
 			for (FBPVariableDescription const& LocalVar : FunctionEntry->LocalVariables)
 			{
 				// Create a member reference so we can safely resolve the UProperty
 				FMemberReference Reference;
 				Reference.SetLocalMember(LocalVar.VarName, FunctionGraph->GetName(), LocalVar.VarGuid);
 
-				UBlueprintNodeSpawner* GetVarSpawner = UBlueprintVariableNodeSpawner::Create(UK2Node_VariableGet::StaticClass(), FunctionGraph, LocalVar, Reference.ResolveMember<UProperty>(Blueprint->SkeletonGeneratedClass));
+				UBlueprintNodeSpawner* GetVarSpawner = UBlueprintVariableNodeSpawner::CreateFromLocal(UK2Node_VariableGet::StaticClass(), FunctionGraph, LocalVar, Reference.ResolveMember<UProperty>(Blueprint->SkeletonGeneratedClass));
 				ActionListOut.Add(GetVarSpawner);
-				UBlueprintNodeSpawner* SetVarSpawner = UBlueprintVariableNodeSpawner::Create(UK2Node_VariableSet::StaticClass(), FunctionGraph, LocalVar, Reference.ResolveMember<UProperty>(Blueprint->SkeletonGeneratedClass));
+				UBlueprintNodeSpawner* SetVarSpawner = UBlueprintVariableNodeSpawner::CreateFromLocal(UK2Node_VariableSet::StaticClass(), FunctionGraph, LocalVar, Reference.ResolveMember<UProperty>(Blueprint->SkeletonGeneratedClass));
 				ActionListOut.Add(SetVarSpawner);
 			}
 		}
diff --git a/Engine/Source/Editor/BlueprintGraph/Private/BlueprintFunctionNodeSpawner.cpp b/Engine/Source/Editor/BlueprintGraph/Private/BlueprintFunctionNodeSpawner.cpp
--- a/Engine/Source/Editor/BlueprintGraph/Private/BlueprintFunctionNodeSpawner.cpp
+++ b/Engine/Source/Editor/BlueprintGraph/Private/BlueprintFunctionNodeSpawner.cpp
@@ -89,7 +89,7 @@ static bool BlueprintFunctionNodeSpawnerImpl::BindFunctionNode(UK2Node_CallFunct
 			}
 			else
 			{
-				UBlueprintNodeSpawner* TempNodeSpawner = UBlueprintVariableNodeSpawner::Create(UK2Node_VariableGet::StaticClass(), BoundProperty);
+				UBlueprintNodeSpawner* TempNodeSpawner = UBlueprintVariableNodeSpawner::CreateFromMemberOrParam(UK2Node_VariableGet::StaticClass(), BoundProperty);
 				bSuccessfulBinding = BindFunctionNode<UK2Node_VariableGet>(NewNode, TempNodeSpawner);
 			}
 		}
diff --git a/Engine/Source/Editor/BlueprintGraph/Private/BlueprintVariableNodeSpawner.cpp b/Engine/Source/Editor/BlueprintGraph/Private/BlueprintVariableNodeSpawner.cpp
--- a/Engine/Source/Editor/BlueprintGraph/Private/BlueprintVariableNodeSpawner.cpp
+++ b/Engine/Source/Editor/BlueprintGraph/Private/BlueprintVariableNodeSpawner.cpp
@@ -19,7 +19,7 @@
  ******************************************************************************/
 
 //------------------------------------------------------------------------------
-UBlueprintVariableNodeSpawner* UBlueprintVariableNodeSpawner::Create(TSubclassOf<UK2Node_Variable> NodeClass, UProperty const* VarProperty, UObject* Outer/*= nullptr*/)
+UBlueprintVariableNodeSpawner* UBlueprintVariableNodeSpawner::CreateFromMemberOrParam(TSubclassOf<UK2Node_Variable> NodeClass, UProperty const* VarProperty, UEdGraph* VarContext, UObject* Outer/*= nullptr*/)
 {
 	check(VarProperty != nullptr);
 	if (Outer == nullptr)
@@ -33,7 +33,8 @@ UBlueprintVariableNodeSpawner* UBlueprintVariableNodeSpawner::Create(TSubclassOf
 
 	UBlueprintVariableNodeSpawner* NodeSpawner = NewObject<UBlueprintVariableNodeSpawner>(Outer);
 	NodeSpawner->NodeClass = NodeClass;
-	NodeSpawner->Field     = VarProperty;
+	NodeSpawner->Field = VarProperty;
+	NodeSpawner->LocalVarOuter = VarContext;
 
 	//--------------------------------------
 	// Default UI Signature
@@ -78,10 +79,11 @@ UBlueprintVariableNodeSpawner* UBlueprintVariableNodeSpawner::Create(TSubclassOf
 			UClass* OwnerClass = Property->GetOwnerClass();
 
 			// We need to use a generated class instead of a skeleton class for IsChildOf, so if the OwnerClass has a Blueprint, grab the GeneratedClass
-			bool const bIsSelfContext = Blueprint->SkeletonGeneratedClass->GetAuthoritativeClass() == OwnerClass || Blueprint->SkeletonGeneratedClass->IsChildOf(OwnerClass);
+			const bool bOwnerClassIsSelfContext = (Blueprint->SkeletonGeneratedClass->GetAuthoritativeClass() == OwnerClass) || Blueprint->SkeletonGeneratedClass->IsChildOf(OwnerClass);
+			const bool bIsFunctionVariable = Property->GetOuter()->IsA(UFunction::StaticClass());
 
 			UK2Node_Variable* VarNode = CastChecked<UK2Node_Variable>(NewNode);
-			VarNode->SetFromProperty(Property, bIsSelfContext);
+			VarNode->SetFromProperty(Property, bOwnerClassIsSelfContext && !bIsFunctionVariable);
 		}
 	};
 	NodeSpawner->SetNodeFieldDelegate = FSetNodeFieldDelegate::CreateStatic(MemberVarSetupLambda);
@@ -90,7 +92,7 @@ UBlueprintVariableNodeSpawner* UBlueprintVariableNodeSpawner::Create(TSubclassOf
 }
 
 //------------------------------------------------------------------------------
-UBlueprintVariableNodeSpawner* UBlueprintVariableNodeSpawner::Create(TSubclassOf<UK2Node_Variable> NodeClass, UEdGraph* VarContext, FBPVariableDescription const& VarDesc, UProperty* VarProperty, UObject* Outer/*= nullptr*/)
+UBlueprintVariableNodeSpawner* UBlueprintVariableNodeSpawner::CreateFromLocal(TSubclassOf<UK2Node_Variable> NodeClass, UEdGraph* VarContext, FBPVariableDescription const& VarDesc, UProperty* VarProperty, UObject* Outer/*= nullptr*/)
 {
 	check(VarContext != nullptr);
 	if (Outer == nullptr)
@@ -158,7 +160,7 @@ void UBlueprintVariableNodeSpawner::Prime()
 FBlueprintNodeSignature UBlueprintVariableNodeSpawner::GetSpawnerSignature() const
 {
 	FBlueprintNodeSignature SpawnerSignature(NodeClass);
-	if (IsLocalVariable())
+	if (IsUserLocalVariable())
 	{
 		SpawnerSignature.AddSubObject(LocalVarOuter);
 		static const FName LocalVarSignatureKey(TEXT("LocalVarName"));
@@ -221,7 +223,7 @@ UEdGraphNode* UBlueprintVariableNodeSpawner::Invoke(UEdGraph* ParentGraph, FBind
 		FCustomizeNodeDelegate PostSpawnDelegate = CustomizeNodeDelegate;
 		if (UObject const* LocalVariableOuter = GetVarOuter())
 		{
-			PostSpawnDelegate = FCustomizeNodeDelegate::CreateStatic(LocalVarSetupLambda, LocalVarDesc.VarName, LocalVariableOuter, LocalVarDesc.VarGuid, CustomizeNodeDelegate);
+			PostSpawnDelegate = FCustomizeNodeDelegate::CreateStatic(LocalVarSetupLambda, IsUserLocalVariable() ? LocalVarDesc.VarName : Field->GetFName(), LocalVariableOuter, LocalVarDesc.VarGuid, CustomizeNodeDelegate);
 		}
 
 		NewNode = UBlueprintNodeSpawner::SpawnNode<UEdGraphNode>(NodeClass, ParentGraph, Bindings, Location, PostSpawnDelegate);
@@ -235,11 +237,17 @@ UEdGraphNode* UBlueprintVariableNodeSpawner::Invoke(UEdGraph* ParentGraph, FBind
 }
 
 //------------------------------------------------------------------------------
-bool UBlueprintVariableNodeSpawner::IsLocalVariable() const
+bool UBlueprintVariableNodeSpawner::IsUserLocalVariable() const
 {
 	return (LocalVarDesc.VarName != NAME_None);
 }
 
+//------------------------------------------------------------------------------
+bool UBlueprintVariableNodeSpawner::IsLocalVariable() const
+{
+	return (LocalVarDesc.VarName != NAME_None) || (LocalVarOuter != nullptr);
+}
+
 //------------------------------------------------------------------------------
 UObject const* UBlueprintVariableNodeSpawner::GetVarOuter() const
 {
@@ -266,7 +274,7 @@ UProperty const* UBlueprintVariableNodeSpawner::GetVarProperty() const
 FEdGraphPinType UBlueprintVariableNodeSpawner::GetVarType() const
 {
 	FEdGraphPinType VarType;
-	if (IsLocalVariable())
+	if (IsUserLocalVariable())
 	{
 		VarType = LocalVarDesc.VarType;
 	}
@@ -284,7 +292,7 @@ FText UBlueprintVariableNodeSpawner::GetVariableName() const
 	FText VarName;
 
 	bool bShowFriendlyNames = GetDefault<UEditorStyleSettings>()->bShowFriendlyNames;
-	if (IsLocalVariable())
+	if (IsUserLocalVariable())
 	{
 		VarName = bShowFriendlyNames ? FText::FromString(LocalVarDesc.FriendlyName) : FText::FromName(LocalVarDesc.VarName);
 	}
diff --git a/Engine/Source/Editor/BlueprintGraph/Private/K2Node_Event.cpp b/Engine/Source/Editor/BlueprintGraph/Private/K2Node_Event.cpp
--- a/Engine/Source/Editor/BlueprintGraph/Private/K2Node_Event.cpp
+++ b/Engine/Source/Editor/BlueprintGraph/Private/K2Node_Event.cpp
@@ -358,7 +358,7 @@ void UK2Node_Event::GetRedirectPinNames(const UEdGraphPin& Pin, TArray<FString>&
 
 	if ( RedirectPinNames.Num() > 0 )
 	{
-		const FString& OldPinName = RedirectPinNames[0];
+		const FString OldPinName = RedirectPinNames[0];
 
 		// first add functionname.param
 		RedirectPinNames.Add(FString::Printf(TEXT("%s.%s"), *EventReference.GetMemberName().ToString(), *OldPinName));
diff --git a/Engine/Source/Editor/BlueprintGraph/Private/K2Node_FunctionEntry.cpp b/Engine/Source/Editor/BlueprintGraph/Private/K2Node_FunctionEntry.cpp
--- a/Engine/Source/Editor/BlueprintGraph/Private/K2Node_FunctionEntry.cpp
+++ b/Engine/Source/Editor/BlueprintGraph/Private/K2Node_FunctionEntry.cpp
@@ -161,6 +161,11 @@ public:
 			GenerateSimpleThenGoto(Context, *Node);
 		}
 	}
+
+	virtual bool RequiresRegisterNetsBeforeScheduling() const override
+	{
+		return true;
+	}
 };
 
 struct FFunctionEntryHelper
@@ -314,7 +319,7 @@ void UK2Node_FunctionEntry::GetRedirectPinNames(const UEdGraphPin& Pin, TArray<F
 
 	if(RedirectPinNames.Num() > 0)
 	{
-		const FString& OldPinName = RedirectPinNames[0];
+		const FString OldPinName = RedirectPinNames[0];
 
 		
 		// first add functionname.param
diff --git a/Engine/Source/Editor/BlueprintGraph/Private/K2Node_FunctionResult.cpp b/Engine/Source/Editor/BlueprintGraph/Private/K2Node_FunctionResult.cpp
--- a/Engine/Source/Editor/BlueprintGraph/Private/K2Node_FunctionResult.cpp
+++ b/Engine/Source/Editor/BlueprintGraph/Private/K2Node_FunctionResult.cpp
@@ -109,6 +109,11 @@ public:
 			GotoStatement.Type = KCST_GotoReturn;
 		}
 	}
+
+	virtual bool RequiresRegisterNetsBeforeScheduling() const override
+	{
+		return true;
+	}
 };
 
 UK2Node_FunctionResult::UK2Node_FunctionResult(const FObjectInitializer& ObjectInitializer)
diff --git a/Engine/Source/Editor/BlueprintGraph/Private/K2Node_GetClassDefaults.cpp b/Engine/Source/Editor/BlueprintGraph/Private/K2Node_GetClassDefaults.cpp
--- a/Engine/Source/Editor/BlueprintGraph/Private/K2Node_GetClassDefaults.cpp
+++ b/Engine/Source/Editor/BlueprintGraph/Private/K2Node_GetClassDefaults.cpp
@@ -23,11 +23,12 @@ namespace
 	// Optional pin manager subclass.
 	struct FClassDefaultsOptionalPinManager : public FOptionalPinManager
 	{
-		FClassDefaultsOptionalPinManager(UClass* InClass, bool bExcludeObjectArrays)
+		FClassDefaultsOptionalPinManager(UClass* InClass, bool bExcludeObjectContainers, bool bExcludeObjectArrays)
 			:FOptionalPinManager()
 		{
 			SrcClass = InClass;
-			bExcludeObjectArrayProperties = bExcludeObjectArrays;
+			bExcludeObjectArrayProperties = bExcludeObjectContainers | bExcludeObjectArrays;
+			bExcludeObjectContainerProperties = bExcludeObjectContainers;
 		}
 
 		virtual void GetRecordDefaults(UProperty* TestProperty, FOptionalPinFromProperty& Record) const override
@@ -54,6 +55,20 @@ namespace
 					TestProperty = TestArrayProperty->Inner;
 				}
 			}
+			else if (USetProperty* TestSetProperty = Cast<USetProperty>(TestProperty))
+			{
+				if (bExcludeObjectContainerProperties && TestSetProperty->ElementProp)
+				{
+					TestProperty = TestSetProperty->ElementProp;
+				}
+			}
+			else if (UMapProperty* TestMapProperty = Cast<UMapProperty>(TestProperty))
+			{
+				// Since we can't treat the key or value as read-only right now, we exclude any TMap that has a non-class UObject reference as its key or value type.
+				return !(bExcludeObjectContainerProperties
+					&& ((TestMapProperty->KeyProp && TestMapProperty->KeyProp->IsA<UObjectProperty>() && !TestMapProperty->KeyProp->IsA<UClassProperty>())
+						|| (TestMapProperty->ValueProp && TestMapProperty->ValueProp->IsA<UObjectProperty>() && !TestMapProperty->ValueProp->IsA<UClassProperty>())));
+			}
 
 			// Don't expose object properties (except for those containing class objects).
 			// @TODO - Could potentially expose object reference values if/when we have support for 'const' input pins.
@@ -72,8 +87,11 @@ namespace
 		// Class type for which optional pins are being managed.
 		UClass* SrcClass;
 
-		// Indicates whether or not object array properties will be excluded.
+		// Indicates whether or not object array properties will be excluded (for backwards-compatibility).
 		bool bExcludeObjectArrayProperties;
+
+		// Indicates whether or not object container properties will be excluded (will supercede the array-specific flag when true).
+		bool bExcludeObjectContainerProperties;
 	};
 
 	// Compilation handler subclass.
@@ -200,9 +218,9 @@ void UK2Node_GetClassDefaults::AllocateDefaultPins()
 
 void UK2Node_GetClassDefaults::PostPlacedNewNode()
 {
-	// Always exclude object array properties for new nodes.
+	// Always exclude object container properties for new nodes.
 	// @TODO - Could potentially expose object reference values if/when we have support for 'const' input pins.
-	bExcludeObjectArrays = true;
+	bExcludeObjectContainers = true;
 
 	if(UEdGraphPin* ClassPin = FindClassPin(Pins))
 	{
@@ -240,20 +258,33 @@ void UK2Node_GetClassDefaults::ValidateNodeDuringCompilation(class FCompilerResu
 {
 	Super::ValidateNodeDuringCompilation(MessageLog);
 
-	if(auto SourceClass = GetInputClass())
+	if (const UClass* SourceClass = GetInputClass())
 	{
-		for(auto Pin : Pins)
+		for (const UEdGraphPin* Pin : Pins)
 		{
-			if(Pin && Pin->Direction == EGPD_Output && Pin->LinkedTo.Num() > 0)
+			// Emit a warning for existing connections to potentially unsafe array property defaults. We do this rather than just implicitly breaking the connection (for compatibility).
+			if (Pin && Pin->Direction == EGPD_Output && Pin->LinkedTo.Num() > 0)
 			{
-				// Emit a warning for existing connections to potentially unsafe array property defaults. We do this rather than just implicitly breaking the connection (for compatibility).
-				if(auto ArrayProperty = Cast<UArrayProperty>(SourceClass->FindPropertyByName(FName(*Pin->PinName))))
+				// Even though container property defaults are copied, the copy could still contain a reference to a non-class object that belongs to the CDO, which would potentially be unsafe to modify.
+				bool bEmitWarning = false;
+				const UProperty* TestProperty = SourceClass->FindPropertyByName(FName(*Pin->PinName));
+				if (const UArrayProperty* ArrayProperty = Cast<UArrayProperty>(TestProperty))
 				{
-					// Even though array property defaults are copied, the copy could still contain a reference to a non-class object that belongs to the CDO, which would potentially be unsafe to modify.
-					if(ArrayProperty->Inner && ArrayProperty->Inner->IsA<UObjectProperty>() && !ArrayProperty->Inner->IsA<UClassProperty>())
-					{
-						MessageLog.Warning(*LOCTEXT("UnsafeConnectionWarning", "@@ has an unsafe connection to the @@ output pin that is not fully supported at this time. It should be disconnected to avoid potentially corrupting class defaults at runtime. If you need to keep this connection, make sure you're not changing the state of any elements in the array. Also note that if you recreate this node, it will not include this output pin.").ToString(), this, Pin);
-					}
+					bEmitWarning = ArrayProperty->Inner && ArrayProperty->Inner->IsA<UObjectProperty>() && !ArrayProperty->Inner->IsA<UClassProperty>();
+				}
+				else if (const USetProperty* SetProperty = Cast<USetProperty>(TestProperty))
+				{
+					bEmitWarning = SetProperty->ElementProp && SetProperty->ElementProp->IsA<UObjectProperty>() && !SetProperty->ElementProp->IsA<UClassProperty>();
+				}
+				else if (const UMapProperty* MapProperty = Cast<UMapProperty>(TestProperty))
+				{
+					bEmitWarning = (MapProperty->KeyProp && MapProperty->KeyProp->IsA<UObjectProperty>() && !MapProperty->KeyProp->IsA<UClassProperty>())
+						|| (MapProperty->ValueProp && MapProperty->ValueProp->IsA<UObjectProperty>() && !MapProperty->ValueProp->IsA<UClassProperty>());
+				}
+
+				if (bEmitWarning)
+				{
+					MessageLog.Warning(*LOCTEXT("UnsafeConnectionWarning", "@@ has an unsafe connection to the @@ output pin that is not fully supported at this time. It should be disconnected to avoid potentially corrupting class defaults at runtime. If you need to keep this connection, make sure you're not changing the state of any elements in the container. Also note that if you recreate this node, it will not include this output pin.").ToString(), this, Pin);
 				}
 			}
 		}
@@ -298,13 +329,13 @@ void UK2Node_GetClassDefaults::ExpandNode(class FKismetCompilerContext& Compiler
 	const UClass* ClassType = GetInputClass();
 
 	// @TODO - Remove if/when we support 'const' input pins.
-	// For array properties, return a local copy of the array so that the original cannot be modified.
+	// For container properties, return a local copy of the container so that the original cannot be modified.
 	for(UEdGraphPin* OutputPin : Pins)
 	{
 		if(OutputPin != nullptr && OutputPin->Direction == EGPD_Output && OutputPin->LinkedTo.Num() > 0)
 		{
 			UProperty* BoundProperty = FindField<UProperty>(ClassType, *(OutputPin->PinName));
-			if(BoundProperty != nullptr && BoundProperty->IsA<UArrayProperty>())
+			if(BoundProperty != nullptr && (BoundProperty->IsA<UArrayProperty>() || BoundProperty->IsA<USetProperty>() || BoundProperty->IsA<UMapProperty>()))
 			{
 				UK2Node_TemporaryVariable* LocalVariable = CompilerContext.SpawnIntermediateNode<UK2Node_TemporaryVariable>(this, SourceGraph);
 				LocalVariable->VariableType = OutputPin->PinType;
@@ -419,7 +450,7 @@ void UK2Node_GetClassDefaults::OnBlueprintClassModified(UBlueprint* TargetBluepr
 void UK2Node_GetClassDefaults::CreateOutputPins(UClass* InClass)
 {
 	// Create the set of output pins through the optional pin manager
-	FClassDefaultsOptionalPinManager OptionalPinManager(InClass, bExcludeObjectArrays);
+	FClassDefaultsOptionalPinManager OptionalPinManager(InClass, bExcludeObjectContainers, bExcludeObjectArrays_DEPRECATED);
 	OptionalPinManager.RebuildPropertyList(ShowPinForProperties, InClass);
 	OptionalPinManager.CreateVisiblePins(ShowPinForProperties, InClass, EGPD_Output, this);
 
diff --git a/Engine/Source/Editor/DetailCustomizations/Private/CurveStructCustomization.cpp b/Engine/Source/Editor/DetailCustomizations/Private/CurveStructCustomization.cpp
--- a/Engine/Source/Editor/DetailCustomizations/Private/CurveStructCustomization.cpp
+++ b/Engine/Source/Editor/DetailCustomizations/Private/CurveStructCustomization.cpp
@@ -192,6 +192,7 @@ void FCurveStructCustomization::CustomizeChildren( TSharedRef<IPropertyHandle> I
 							.ToolTipText(LOCTEXT( "CreateAssetTooltip", "Create a new CurveFloat asset from this curve") )
 							.OnClicked(this, &FCurveStructCustomization::OnCreateButtonClicked)
 							.IsEnabled(this, &FCurveStructCustomization::IsCreateButtonEnabled)
+							.Visibility(Owner != nullptr ? EVisibility::Visible : EVisibility::Collapsed)
 						]
 					]
 				];
diff --git a/Engine/Source/Editor/Kismet/Private/BlueprintCompilationManager.cpp b/Engine/Source/Editor/Kismet/Private/BlueprintCompilationManager.cpp
new file mode 100644
--- /dev/null
+++ b/Engine/Source/Editor/Kismet/Private/BlueprintCompilationManager.cpp
@@ -0,0 +1,135 @@
+// Copyright 1998-2016 Epic Games, Inc. All Rights Reserved.
+
+#include "BlueprintCompilationManager.h"
+#include "BlueprintEditorUtils.h"
+#include "Kismet2/KismetEditorUtilities.h"
+#include "Kismet2/KismetReinstanceUtilities.h"
+
+struct FBlueprintCompilationManagerImpl : public FGCObject, FTickableEditorObject
+{
+	virtual ~FBlueprintCompilationManagerImpl() {}
+
+	// FGCObject:
+	virtual void AddReferencedObjects(FReferenceCollector& Collector);
+
+	// FTickableEditorObject
+	virtual void Tick(float DeltaTime) override { TickCompilationQueue(); }
+	virtual bool IsTickable() const override { return true; }
+	virtual TStatId GetStatId() const override { RETURN_QUICK_DECLARE_CYCLE_STAT(FBlueprintCompilationManagerImpl, STATGROUP_Tickables); }
+
+	void QueueForCompilation(UBlueprint* BPToCompile);
+	void TickCompilationQueue();
+	bool HasBlueprintsToCompile();
+
+	TArray<UBlueprint*> QueuedBPs;
+};
+
+void FBlueprintCompilationManagerImpl::AddReferencedObjects(FReferenceCollector& Collector)
+{
+	Collector.AddReferencedObjects(QueuedBPs);
+}
+
+void FBlueprintCompilationManagerImpl::QueueForCompilation(UBlueprint* BPToCompile)
+{
+	QueuedBPs.Add(BPToCompile);
+}
+
+void FBlueprintCompilationManagerImpl::TickCompilationQueue()
+{
+	TArray<UBlueprint*> QueuedBPsLocal = MoveTemp(QueuedBPs);
+	QueuedBPsLocal.Sort(
+		[](UBlueprint& A, UBlueprint& B)
+	{
+		int32 DepthA = 0;
+		int32 DepthB = 0;
+		UStruct* Iter = *(A.GeneratedClass) ? A.GeneratedClass->GetSuperStruct() : nullptr;
+		while (Iter)
+		{
+			++DepthA;
+			Iter = Iter->GetSuperStruct();
+		}
+
+		Iter = *(B.GeneratedClass) ? B.GeneratedClass->GetSuperStruct() : nullptr;
+		while (Iter)
+		{
+			++DepthB;
+			Iter = Iter->GetSuperStruct();
+		}
+
+		if (DepthA == DepthB)
+		{
+			return A.GetFName() < B.GetFName(); 
+		}
+		return DepthA < DepthB;
+	}
+	);
+
+	// Hack - fixes an issue when compiling class hierarchies. Reinstancer recompiles children before they have
+	// had their nodes reconstructed. Same for null graphs:
+	for (UBlueprint* BP : QueuedBPsLocal)
+	{
+		UPackage* Package = Cast<UPackage>(BP->GetOutermost());
+		bool bIsPackageDirty = Package ? Package->IsDirty() : false;
+
+		FBlueprintEditorUtils::PurgeNullGraphs(BP);
+		FBlueprintEditorUtils::ReconstructAllNodes(BP);
+
+		if( Package )
+		{
+			Package->SetDirtyFlag(bIsPackageDirty);
+		}
+	}
+
+	TArray<UObject*> ObjsLoaded;
+	for (UBlueprint* BP : QueuedBPsLocal)
+	{
+		FBlueprintEditorUtils::RegenerateBlueprintClass(BP, BP->GeneratedClass, BP->GeneratedClass->ClassDefaultObject, ObjsLoaded);
+	}
+}
+
+bool FBlueprintCompilationManagerImpl::HasBlueprintsToCompile()
+{
+	return QueuedBPs.Num() != 0;
+}
+
+// Singleton boilperplate:
+FBlueprintCompilationManagerImpl* BPCMImpl = nullptr;
+
+void FBlueprintCompilationManager::Initialize()
+{
+	if(!BPCMImpl)
+	{
+		BPCMImpl = new FBlueprintCompilationManagerImpl();
+	}
+}
+
+void FBlueprintCompilationManager::Shutdown()
+{
+	delete BPCMImpl;
+	BPCMImpl = nullptr;
+}
+
+FBlueprintCompilationManagerImpl* FBlueprintCompilationManager::DebugAccessImpl()
+{
+	return BPCMImpl;
+}
+
+// Forward to impl:
+void FBlueprintCompilationManager::QueueForCompilation(UBlueprint* BPToCompile)
+{
+	check(BPToCompile->GetLinker());
+	BPCMImpl->QueueForCompilation(BPToCompile);
+}
+	
+void FBlueprintCompilationManager::NotifyBlueprintLoaded(UBlueprint* BPLoaded)
+{
+	// Blueprints can be loaded before editor modules are on line:
+	if(!BPCMImpl)
+	{
+		FBlueprintCompilationManager::Initialize();
+	}
+	check(BPLoaded->GetLinker());
+	BPCMImpl->QueueForCompilation(BPLoaded);
+}
+
+
diff --git a/Engine/Source/Editor/Kismet/Private/SMyBlueprint.cpp b/Engine/Source/Editor/Kismet/Private/SMyBlueprint.cpp
--- a/Engine/Source/Editor/Kismet/Private/SMyBlueprint.cpp
+++ b/Engine/Source/Editor/Kismet/Private/SMyBlueprint.cpp
@@ -797,7 +797,9 @@ bool SMyBlueprint::CanRequestRenameOnActionNode(TWeakPtr<FGraphActionNode> InSel
 				// Do not allow renaming of any graph actions outside of the following
 				if(Actions[i]->GetSectionID() == NodeSectionID::FUNCTION || Actions[i]->GetSectionID() == NodeSectionID::MACRO)
 				{
-					if(!FBlueprintEditorUtils::IsPaletteActionReadOnly(Actions[i], BlueprintEditorPtr.Pin()) && FBlueprintEditorUtils::FindBlueprintForGraph(GraphAction->EdGraph) == GetBlueprintObj())
+					if(!FBlueprintEditorUtils::IsPaletteActionReadOnly(Actions[i], BlueprintEditorPtr.Pin())
+						&& FBlueprintEditorUtils::FindBlueprintForGraph(GraphAction->EdGraph) == GetBlueprintObj()
+						&& FBlueprintEditorUtils::FindFunctionInImplementedInterfaces(GetBlueprintObj(), GraphAction->EdGraph->GetFName(), nullptr, true) == nullptr)
 					{
 						bIsReadOnly = false;
 						break;
diff --git a/Engine/Source/Editor/Kismet/Private/SSCSEditor.cpp b/Engine/Source/Editor/Kismet/Private/SSCSEditor.cpp
--- a/Engine/Source/Editor/Kismet/Private/SSCSEditor.cpp
+++ b/Engine/Source/Editor/Kismet/Private/SSCSEditor.cpp
@@ -95,7 +95,8 @@ void SSCSEditorDragDropTree::Construct( const FArguments& InArgs )
 			.SelectionMode( InArgs._SelectionMode )
 			.HeaderRow( InArgs._HeaderRow )
 			.ClearSelectionOnClick( InArgs._ClearSelectionOnClick )
-			.ExternalScrollbar( InArgs._ExternalScrollbar );
+			.ExternalScrollbar( InArgs._ExternalScrollbar )
+			.OnEnteredBadState( InArgs._OnTableViewBadState );
 
 	STreeView<FSCSEditorTreeNodePtrType>::Construct( BaseArgs );
 }
@@ -3337,6 +3338,7 @@ void SSCSEditor::Construct( const FArguments& InArgs )
 		.OnItemScrolledIntoView(this, &SSCSEditor::OnItemScrolledIntoView)
 		.OnMouseButtonDoubleClick(this, &SSCSEditor::HandleItemDoubleClicked)
 		.ClearSelectionOnClick(InArgs._EditorMode == EComponentEditorMode::BlueprintSCS ? true : false)
+		.OnTableViewBadState(this, &SSCSEditor::DumpTree)
 		.ItemHeight(24)
 		.HeaderRow
 		(
@@ -4428,6 +4430,169 @@ void SSCSEditor::UpdateTree(bool bRegenerateTreeNodes)
 	SCSTreeWidget->RequestTreeRefresh();
 }
 
+void SSCSEditor::DumpTree()
+{
+	/* Example:
+
+		[ACTOR] MyBlueprint (self)
+		|
+		[SEPARATOR]
+		|
+		DefaultSceneRoot (Inherited)
+		|
+		+- StaticMesh (Inherited)
+		|  |
+		|  +- Scene4 (Inherited)
+		|  |
+		|  +- Scene (Inherited)
+		|     |
+		|     +- Scene1 (Inherited)
+		|  
+		+- Scene2 (Inherited)
+		|  |
+		|  +- Scene3 (Inherited)
+		|
+		[SEPARATOR]
+		|
+		ProjectileMovement (Inherited)
+	*/
+
+	UE_LOG(LogSCSEditor, Log, TEXT("---------------------"));
+	UE_LOG(LogSCSEditor, Log, TEXT(" STreeView NODE DUMP"));
+	UE_LOG(LogSCSEditor, Log, TEXT("---------------------"));
+
+	const UBlueprint* BlueprintContext = nullptr;
+	const AActor* ActorInstance = GetActorContext();
+	if (ActorInstance)
+	{
+		BlueprintContext = UBlueprint::GetBlueprintFromClass(ActorInstance->GetClass());
+	}
+
+	TArray<TArray<FSCSEditorTreeNodePtrType>> NodeListStack;
+	NodeListStack.Push(RootNodes);
+
+	auto LineSpacingLambda = [&NodeListStack](const TArray<FSCSEditorTreeNodePtrType>& NodeList, int32 CurrentDepth, const FString& Prefix)
+	{
+		bool bAddLineSpacing = false;
+		for (int Depth = 0; Depth <= CurrentDepth && !bAddLineSpacing; ++Depth)
+		{
+			bAddLineSpacing = NodeListStack[Depth].Num() > 0;
+		}
+
+		if (bAddLineSpacing)
+		{
+			UE_LOG(LogSCSEditor, Log, TEXT(" %s%s"), *Prefix, NodeList.Num() > 0 ? TEXT("|") : TEXT(""));
+		}
+	};
+
+	while (NodeListStack.Num() > 0)
+	{
+		const int32 CurrentDepth = NodeListStack.Num() - 1;
+		TArray<FSCSEditorTreeNodePtrType>& NodeList = NodeListStack[CurrentDepth];
+		if (NodeList.Num() > 0)
+		{
+			FString Prefix;
+			for (int32 Depth = 1; Depth < CurrentDepth; ++Depth)
+			{
+				int32 NodeCount = NodeListStack[Depth].Num();
+				if (Depth == 1)
+				{
+					NodeCount += NodeListStack[0].Num();
+				}
+
+				Prefix += (NodeCount > 0) ? TEXT("|  ") : TEXT("   ");
+			}
+
+			FString NodePrefix;
+			if (CurrentDepth > 0)
+			{
+				NodePrefix = TEXT("+- ");
+			}
+
+			FSCSEditorTreeNodePtrType Node = NodeList[0];
+			NodeList.RemoveAt(0);
+
+			if (Node.IsValid())
+			{
+				FString NodeLabel = TEXT("[UNKNOWN]");
+				switch (Node->GetNodeType())
+				{
+				case FSCSEditorTreeNode::ENodeType::RootActorNode:
+					switch (EditorMode)
+					{
+					case EComponentEditorMode::ActorInstance:
+						NodeLabel = TEXT("[ACTOR]");
+						break;
+
+					case EComponentEditorMode::BlueprintSCS:
+						NodeLabel = TEXT("[BLUEPRINT]");
+						break;
+					}
+
+					if (BlueprintContext)
+					{
+						NodeLabel += FString::Printf(TEXT(" %s (self)"), *BlueprintContext->GetName());
+					}
+					else if (ActorInstance)
+					{
+						NodeLabel += FString::Printf(TEXT(" %s (Instance)"), *ActorInstance->GetActorLabel());
+					}
+					break;
+
+				case FSCSEditorTreeNode::ENodeType::SeparatorNode:
+					NodeLabel = TEXT("[SEPARATOR]");
+					break;
+
+				case FSCSEditorTreeNode::ENodeType::ComponentNode:
+					NodeLabel = Node->GetDisplayString();
+					if (Node->IsInherited())
+					{
+						NodeLabel += TEXT(" (Inherited)");
+					}
+					break;
+				}
+
+				UE_LOG(LogSCSEditor, Log, TEXT(" %s%s%s"), *Prefix, *NodePrefix, *NodeLabel);
+
+				const TArray<FSCSEditorTreeNodePtrType>& Children = Node->GetChildren();
+				if (Children.Num() > 0)
+				{
+					if (CurrentDepth > 1)
+					{
+						UE_LOG(LogSCSEditor, Log, TEXT(" %s%s|"), *Prefix, NodeListStack[CurrentDepth].Num() > 0 ? TEXT("|  ") : TEXT("   "));
+					}
+					else if (CurrentDepth == 1)
+					{
+						UE_LOG(LogSCSEditor, Log, TEXT(" %s%s|"), *Prefix, NodeListStack[0].Num() > 0 ? TEXT("|  ") : TEXT("   "));
+					}
+					else
+					{
+						UE_LOG(LogSCSEditor, Log, TEXT(" %s|"), *Prefix);
+					}
+
+					NodeListStack.Push(Children);
+				}
+				else
+				{
+					LineSpacingLambda(NodeList, CurrentDepth, Prefix);
+				}
+			}
+			else
+			{
+				UE_LOG(LogSCSEditor, Log, TEXT(" %s%s[INVALID]"), *Prefix, *NodePrefix);
+				
+				LineSpacingLambda(NodeList, CurrentDepth, Prefix);
+			}
+		}
+		else
+		{
+			NodeListStack.Pop();
+		}
+	}
+
+	UE_LOG(LogSCSEditor, Log, TEXT("--------(end)--------"));
+}
+
 const TArray<FSCSEditorTreeNodePtrType>& SSCSEditor::GetRootNodes() const
 {
 	return RootNodes;
diff --git a/Engine/Source/Editor/Kismet/Public/BlueprintCompilationManager.h b/Engine/Source/Editor/Kismet/Public/BlueprintCompilationManager.h
new file mode 100644
--- /dev/null
+++ b/Engine/Source/Editor/Kismet/Public/BlueprintCompilationManager.h
@@ -0,0 +1,16 @@
+// Copyright 1998-2016 Epic Games, Inc. All Rights Reserved.
+
+#pragma once
+
+struct FBlueprintCompilationManagerImpl;
+
+struct KISMET_API FBlueprintCompilationManager
+{
+	static void Initialize();
+	static void Shutdown();
+	static FBlueprintCompilationManagerImpl* DebugAccessImpl();
+
+	static void QueueForCompilation(UBlueprint* BPToCompile);
+	static void NotifyBlueprintLoaded(UBlueprint* BPLoaded);
+};
+
diff --git a/Engine/Source/Editor/Kismet/Public/SSCSEditor.h b/Engine/Source/Editor/Kismet/Public/SSCSEditor.h
--- a/Engine/Source/Editor/Kismet/Public/SSCSEditor.h
+++ b/Engine/Source/Editor/Kismet/Public/SSCSEditor.h
@@ -667,6 +667,7 @@ public:
 		, _SelectionMode(ESelectionMode::Multi)
 		, _ClearSelectionOnClick(true)
 		, _ExternalScrollbar()
+		, _OnTableViewBadState()
 		{}
 
 		SLATE_ARGUMENT( SSCSEditor*, SCSEditor )
@@ -699,6 +700,8 @@ public:
 
 		SLATE_ARGUMENT( TSharedPtr<SScrollBar>, ExternalScrollbar )
 
+		SLATE_EVENT( FOnTableViewBadState, OnTableViewBadState )
+
 	SLATE_END_ARGS()
 	/** Object construction - mostly defers to the base STreeView */
 	void Construct( const FArguments& InArgs );
@@ -842,6 +845,9 @@ public:
 	/** Refresh the tree control to reflect changes in the SCS */
 	void UpdateTree(bool bRegenerateTreeNodes = true);
 
+	/** Dumps out the tree view contents to the log (used to assist with debugging widget hierarchy issues) */
+	void DumpTree();
+
 	/** Forces the details panel to refresh on the same objects */
 	void RefreshSelectionDetails();
 
diff --git a/Engine/Source/Editor/KismetCompiler/Private/KismetCompiler.cpp b/Engine/Source/Editor/KismetCompiler/Private/KismetCompiler.cpp
--- a/Engine/Source/Editor/KismetCompiler/Private/KismetCompiler.cpp
+++ b/Engine/Source/Editor/KismetCompiler/Private/KismetCompiler.cpp
@@ -54,6 +54,7 @@
 #include "Serialization/ArchiveScriptReferenceCollector.h"
 
 static bool bDebugPropertyPropagation = false;
+extern COREUOBJECT_API bool GMinimalCompileOnLoad;
 
 #define USE_TRANSIENT_SKELETON 0
 
@@ -236,7 +237,9 @@ void FKismetCompilerContext::CleanAndSanitizeClass(UBlueprintGeneratedClass* Cla
 	NewClass = ClassToClean;
 	OldCDO = ClassToClean->ClassDefaultObject; // we don't need to create the CDO at this point
 	
-	const ERenameFlags RenFlags = REN_DontCreateRedirectors | (bRecompilingOnLoad ? REN_ForceNoResetLoaders : 0) | REN_NonTransactional | REN_DoNotDirty;
+	// when doing the GMinimalCompileOnLoad path we ned to make sure that loaders aren't reset
+	// because the reinstancer may run a compile pass with bRecompilingOnLoad set to false)
+	const ERenameFlags RenFlags = REN_DontCreateRedirectors |  ((bRecompilingOnLoad || GMinimalCompileOnLoad) ? REN_ForceNoResetLoaders : 0) | REN_NonTransactional | REN_DoNotDirty;
 
 	if( OldCDO )
 	{
@@ -868,72 +871,66 @@ static void SwapElementsInSingleLinkedList(UField* & PtrToFirstElement, UField*
 	PtrToFirstElement = TempSecond;
 }
 
-void FKismetCompilerContext::CreateLocalVariablesForFunction(FKismetFunctionContext& Context, UFunction* ParameterSignature)
+void FKismetCompilerContext::CreateParametersForFunction(FKismetFunctionContext& Context, UFunction* ParameterSignature, UField**& FunctionPropertyStorageLocation)
 {
-	ensure(Context.IsEventGraph() || !Context.EventGraphLocals.Num());
-	ensure(!Context.IsEventGraph() || !Context.Locals.Num() || !UsePersistentUberGraphFrame());
+	const bool bArePropertiesLocal = true;
+	CreatePropertiesFromList(Context.Function, FunctionPropertyStorageLocation, Context.Parameters, CPF_Parm, bArePropertiesLocal, /*bPropertiesAreParameters=*/ true);
+	CreatePropertiesFromList(Context.Function, FunctionPropertyStorageLocation, Context.Results, CPF_Parm | CPF_OutParm, bArePropertiesLocal, /*bPropertiesAreParameters=*/ true);
 
-	const bool bPersistentUberGraphFrame = UsePersistentUberGraphFrame() && Context.bIsUbergraph;
-	// Local stack frame (or maybe class for the ubergraph)
+	//MAKE SURE THE PARAMETERS ORDER MATCHES THE OVERRIDEN FUNCTION
+	if (ParameterSignature)
 	{
-		const bool bArePropertiesLocal = true;
-
-		// Pull the local properties generated out of the function, they will be put at the end of the list
-		UField* LocalProperties = Context.Function->Children;
-		Context.Function->Children = nullptr;
-
-		UField** PropertyStorageLocation = &Context.Function->Children;
-		CreatePropertiesFromList(Context.Function, PropertyStorageLocation, Context.Parameters, CPF_Parm, bArePropertiesLocal, /*bPropertiesAreParameters=*/ true);
-		CreatePropertiesFromList(Context.Function, PropertyStorageLocation, Context.Results, CPF_Parm | CPF_OutParm, bArePropertiesLocal, /*bPropertiesAreParameters=*/ true);
-
-		//MAKE SURE THE PARAMETERS ORDER MATCHES THE OVERRIDEN FUNCTION
-		if (ParameterSignature)
+		UField** CurrentFieldStorageLocation = &Context.Function->Children;
+		for (TFieldIterator<UProperty> SignatureIt(ParameterSignature); SignatureIt && ((SignatureIt->PropertyFlags & CPF_Parm) != 0); ++SignatureIt)
 		{
-			UField** CurrentFieldStorageLocation = &Context.Function->Children;
-			for (TFieldIterator<UProperty> SignatureIt(ParameterSignature); SignatureIt && ((SignatureIt->PropertyFlags & CPF_Parm) != 0); ++SignatureIt)
+			const FName WantedName = SignatureIt->GetFName();
+			if (!*CurrentFieldStorageLocation || (WantedName != (*CurrentFieldStorageLocation)->GetFName()))
 			{
-				const FName WantedName = SignatureIt->GetFName();
-				if (!*CurrentFieldStorageLocation || (WantedName != (*CurrentFieldStorageLocation)->GetFName()))
+				//Find Field with the proper name
+				UField** FoundFieldStorageLocation = *CurrentFieldStorageLocation ? &((*CurrentFieldStorageLocation)->Next) : nullptr;
+				while (FoundFieldStorageLocation && *FoundFieldStorageLocation && (WantedName != (*FoundFieldStorageLocation)->GetFName()))
 				{
-					//Find Field with the proper name
-					UField** FoundFieldStorageLocation = *CurrentFieldStorageLocation ? &((*CurrentFieldStorageLocation)->Next) : nullptr;
-					while (FoundFieldStorageLocation && *FoundFieldStorageLocation && (WantedName != (*FoundFieldStorageLocation)->GetFName()))
-					{
-						FoundFieldStorageLocation = &((*FoundFieldStorageLocation)->Next);
-					}
+					FoundFieldStorageLocation = &((*FoundFieldStorageLocation)->Next);
+				}
 
-					if (FoundFieldStorageLocation && *FoundFieldStorageLocation)
-					{
-						// swap the found field and OverridenIt
-						SwapElementsInSingleLinkedList(*CurrentFieldStorageLocation, *FoundFieldStorageLocation); //FoundFieldStorageLocation points now a random element 
-					}
-					else
-					{
-						MessageLog.Error(*FString::Printf(*LOCTEXT("WrongParameterOrder_Error", "Cannot order parameters %s in function %s.").ToString(), *WantedName.ToString(), *Context.Function->GetName()));
-						break;
-					}
+				if (FoundFieldStorageLocation && *FoundFieldStorageLocation)
+				{
+					// swap the found field and OverridenIt
+					SwapElementsInSingleLinkedList(*CurrentFieldStorageLocation, *FoundFieldStorageLocation); //FoundFieldStorageLocation points now a random element 
+				}
+				else
+				{
+					MessageLog.Error(*FString::Printf(*LOCTEXT("WrongParameterOrder_Error", "Cannot order parameters %s in function %s.").ToString(), *WantedName.ToString(), *Context.Function->GetName()));
+					break;
 				}
-				CurrentFieldStorageLocation = &((*CurrentFieldStorageLocation)->Next);
-			}
-			PropertyStorageLocation = CurrentFieldStorageLocation;
-			// There is no guarantee that CurrentFieldStorageLocation points the last parameter's next. We need to ensure that.
-			while (*PropertyStorageLocation)
-			{
-				PropertyStorageLocation = &((*PropertyStorageLocation)->Next);
 			}
+			CurrentFieldStorageLocation = &((*CurrentFieldStorageLocation)->Next);
 		}
+		FunctionPropertyStorageLocation = CurrentFieldStorageLocation;
 
-		CreatePropertiesFromList(Context.Function, PropertyStorageLocation, Context.Locals, 0, bArePropertiesLocal, /*bPropertiesAreParameters=*/ true);
-
-		if (bPersistentUberGraphFrame)
+		// There is no guarantee that CurrentFieldStorageLocation points the last parameter's next. We need to ensure that.
+		while (*FunctionPropertyStorageLocation)
 		{
-			CreatePropertiesFromList(Context.Function, PropertyStorageLocation, Context.EventGraphLocals, 0, bArePropertiesLocal, true);
+			FunctionPropertyStorageLocation = &((*FunctionPropertyStorageLocation)->Next);
 		}
+	}
+}
+
+void FKismetCompilerContext::CreateLocalVariablesForFunction(FKismetFunctionContext& Context, UField**& FunctionPropertyStorageLocation)
+{
+	ensure(Context.IsEventGraph() || !Context.EventGraphLocals.Num());
+	ensure(!Context.IsEventGraph() || !Context.Locals.Num() || !UsePersistentUberGraphFrame());
+
+	const bool bPersistentUberGraphFrame = UsePersistentUberGraphFrame() && Context.bIsUbergraph;
+	// Local stack frame (or maybe class for the ubergraph)
+	{
+		const bool bArePropertiesLocal = true;
 
-		// If there were local properties, place them at the end of the property storage location
-		if(LocalProperties)
+		CreatePropertiesFromList(Context.Function, FunctionPropertyStorageLocation, Context.Locals, 0, bArePropertiesLocal, /*bPropertiesAreParameters=*/ true);
+
+		if (bPersistentUberGraphFrame)
 		{
-			*PropertyStorageLocation = LocalProperties;
+			CreatePropertiesFromList(Context.Function, FunctionPropertyStorageLocation, Context.EventGraphLocals, 0, bArePropertiesLocal, true);
 		}
 
 		// Create debug data for variable reads/writes
@@ -975,34 +972,32 @@ void FKismetCompilerContext::CreateLocalVariablesForFunction(FKismetFunctionCont
 	// Class
 	{
 		int32 PropertySafetyCounter = 100000;
-		UField** PropertyStorageLocation = &(NewClass->Children);
-		while (*PropertyStorageLocation != NULL)
+		UField** ClassPropertyStorageLocation = &(NewClass->Children);
+		while (*ClassPropertyStorageLocation != NULL)
 		{
 			if (--PropertySafetyCounter == 0)
 			{
 				checkf(false, TEXT("Property chain is corrupted;  The most likely causes are multiple properties with the same name.") );
 			}
 
-			PropertyStorageLocation = &((*PropertyStorageLocation)->Next);
+			ClassPropertyStorageLocation = &((*ClassPropertyStorageLocation)->Next);
 		}
 
 		const bool bArePropertiesLocal = false;
 		const uint64 UbergraphHiddenVarFlags = CPF_Transient | CPF_DuplicateTransient;
 		if (!bPersistentUberGraphFrame)
 		{
-			CreatePropertiesFromList(NewClass, PropertyStorageLocation, Context.EventGraphLocals, UbergraphHiddenVarFlags, bArePropertiesLocal);
+			CreatePropertiesFromList(NewClass, ClassPropertyStorageLocation, Context.EventGraphLocals, UbergraphHiddenVarFlags, bArePropertiesLocal);
 		}
 
 		// Handle level actor references
 		const uint64 LevelActorReferenceVarFlags = 0/*CPF_Edit*/;
-		CreatePropertiesFromList(NewClass, PropertyStorageLocation, Context.LevelActorReferences, LevelActorReferenceVarFlags, false);
+		CreatePropertiesFromList(NewClass, ClassPropertyStorageLocation, Context.LevelActorReferences, LevelActorReferenceVarFlags, false);
 	}
 }
 
-void FKismetCompilerContext::CreateUserDefinedLocalVariablesForFunction(FKismetFunctionContext& Context, UField**& PropertyStorageLocation)
+void FKismetCompilerContext::CreateUserDefinedLocalVariablesForFunction(FKismetFunctionContext& Context, UField**& FunctionPropertyStorageLocation)
 {
-	check(Context.Function->Children == NULL);
-
 	// Create local variables from the Context entry point
 	for (int32 i = 0; i < Context.EntryPoint->LocalVariables.Num(); ++i)
 	{
@@ -1013,8 +1008,8 @@ void FKismetCompilerContext::CreateUserDefinedLocalVariablesForFunction(FKismetF
 		if (NewProperty != NULL)
 		{
 			// Link this object to the tail of the list (so properties remain in the desired order)
-			*PropertyStorageLocation = NewProperty;
-			PropertyStorageLocation = &(NewProperty->Next);
+			*FunctionPropertyStorageLocation = NewProperty;
+			FunctionPropertyStorageLocation = &(NewProperty->Next);
 		}
 
 		if (NewProperty != NULL)
@@ -1431,7 +1426,7 @@ void FKismetCompilerContext::PrecompileFunction(FKismetFunctionContext& Context)
 
 		// Create the function stub
 		FName NewFunctionName = (Context.EntryPoint->CustomGeneratedFunctionName != NAME_None) ? Context.EntryPoint->CustomGeneratedFunctionName : Context.EntryPoint->SignatureName;
-		if(Context.IsDelegateSignature())
+		if (Context.IsDelegateSignature())
 		{
 			// prefix with the the blueprint name to avoid conflicts with natively defined delegate signatures
 			FString Name = NewFunctionName.ToString();
@@ -1450,7 +1445,7 @@ void FKismetCompilerContext::PrecompileFunction(FKismetFunctionContext& Context)
 			MessageLog.Error(*FString::Printf(*LOCTEXT("DuplicateFunctionName_Error", "Found more than one function with the same name %s; second occurance at @@").ToString(), *NewFunctionNameString), Context.EntryPoint);
 			return;
 		}
-		else if(NULL != FindField<UProperty>(NewClass, NewFunctionName))
+		else if (NULL != FindField<UProperty>(NewClass, NewFunctionName))
 		{
 			MessageLog.Error(*FString::Printf(*LOCTEXT("DuplicateFieldName_Error", "Name collision - function and property have the same name - '%s'. @@").ToString(), *NewFunctionNameString), Context.EntryPoint);
 			return;
@@ -1477,25 +1472,25 @@ void FKismetCompilerContext::PrecompileFunction(FKismetFunctionContext& Context)
 
 		// Set up the function category
 		FKismetUserDeclaredFunctionMetadata& FunctionMetaData = Context.EntryPoint->MetaData;
-		if( !FunctionMetaData.Category.IsEmpty() )
+		if (!FunctionMetaData.Category.IsEmpty())
 		{
 			Context.Function->SetMetaData(FBlueprintMetadata::MD_FunctionCategory, *FunctionMetaData.Category.ToString());
 		}
 
 		// Set up the function keywords
-		if( !FunctionMetaData.Keywords.IsEmpty() )
+		if (!FunctionMetaData.Keywords.IsEmpty())
 		{
 			Context.Function->SetMetaData(FBlueprintMetadata::MD_FunctionKeywords, *FunctionMetaData.Keywords.ToString());
 		}
 
 		// Set up the function compact node title
-		if( !FunctionMetaData.CompactNodeTitle.IsEmpty() )
+		if (!FunctionMetaData.CompactNodeTitle.IsEmpty())
 		{
 			Context.Function->SetMetaData(FBlueprintMetadata::MD_CompactNodeTitle, *FunctionMetaData.CompactNodeTitle.ToString());
 		}
 
 		// Set as blutility function
-		if( FunctionMetaData.bCallInEditor )
+		if (FunctionMetaData.bCallInEditor)
 		{
 			Context.Function->SetMetaData(FBlueprintMetadata::MD_CallInEditor, TEXT( "true" ));
 		}
@@ -1571,30 +1566,41 @@ void FKismetCompilerContext::PrecompileFunction(FKismetFunctionContext& Context)
 			NewClass->UberGraphFunction = Context.Function;
 		}
 
+		// Register nets from function entry/exit nodes first, even for skeleton compiles (as they form the signature)
+		for (UEdGraphNode* Node : Context.SourceGraph->Nodes)
+		{
+			if (FNodeHandlingFunctor* Handler = NodeHandlers.FindRef(Node->GetClass()))
+			{
+				if (Handler->RequiresRegisterNetsBeforeScheduling())
+				{
+					Handler->RegisterNets(Context, Node);
+				}
+			}
+		}
+
+		UField** FunctionPropertyStorageLocation = &(Context.Function->Children);
+
+		// Create input/output parameter variables, this must occur before registering nets so that the properties are in place
+		CreateParametersForFunction(Context, ParentFunction ? ParentFunction : OverridenFunction, FunctionPropertyStorageLocation);
+
 		// Create any user defined variables, this must occur before registering nets so that the properties are in place
-		UField** PropertyStorageLocation = &(Context.Function->Children);
-		CreateUserDefinedLocalVariablesForFunction(Context, PropertyStorageLocation);
+		CreateUserDefinedLocalVariablesForFunction(Context, FunctionPropertyStorageLocation);
 
 		//@TODO: Prune pure functions that don't have any consumers
 		if (bIsFullCompile)
 		{
-			if (!Blueprint->bIsRegeneratingOnLoad)
-			{
-				BP_SCOPED_COMPILER_EVENT_STAT(EKismetCompilerStats_AnalyzeExecutionPath);
-				FKismetCompilerUtilities::ValidateProperEndExecutionPath(Context);
-			}
-
 			// Find the execution path (and make sure it has no cycles)
 			CreateExecutionSchedule(Context.SourceGraph->Nodes, Context.LinearExecutionList);
 
-			for (int32 NodeIndex = 0; NodeIndex < Context.LinearExecutionList.Num(); ++NodeIndex)
+			// Register nets for any nodes still in the schedule (as long as they didn't get registered in the initial all-nodes pass)
+			for (UEdGraphNode* Node : Context.LinearExecutionList)
 			{
-				UEdGraphNode* Node = Context.LinearExecutionList[NodeIndex];
-
-				// Register nets in the schedule
 				if (FNodeHandlingFunctor* Handler = NodeHandlers.FindRef(Node->GetClass()))
 				{
-					Handler->RegisterNets(Context, Node);
+					if (!Handler->RequiresRegisterNetsBeforeScheduling())
+					{
+						Handler->RegisterNets(Context, Node);
+					}
 				}
 				else
 				{
@@ -1602,30 +1608,15 @@ void FKismetCompilerContext::PrecompileFunction(FKismetFunctionContext& Context)
 				}
 			}
 		}
-		else
-		{
-			for (UEdGraphNode* Node : Context.SourceGraph->Nodes)
-			{
-				const bool bGenerateParameters = Node->IsA<UK2Node_FunctionEntry>();
-				const bool bGenerateResult = Node->IsA<UK2Node_FunctionResult>();
-				if (bGenerateParameters || bGenerateResult)
-				{
-					if (FNodeHandlingFunctor* Handler = NodeHandlers.FindRef(Node->GetClass()))
-					{
-						Handler->RegisterNets(Context, Node);
-					}
-				}
-			}
-		}
-	
-		// Create variable declarations
-		CreateLocalVariablesForFunction(Context, ParentFunction ? ParentFunction : OverridenFunction);
 
-		//Validate AccessSpecifier
+		// Create net variable declarations
+		CreateLocalVariablesForFunction(Context, FunctionPropertyStorageLocation);
+
+		// Validate AccessSpecifier
 		const uint32 AccessSpecifierFlag = FUNC_AccessSpecifiers & Context.EntryPoint->GetExtraFlags();
 		const bool bAcceptedAccessSpecifier = 
 			(0 == AccessSpecifierFlag) || (FUNC_Public == AccessSpecifierFlag) || (FUNC_Protected == AccessSpecifierFlag) || (FUNC_Private == AccessSpecifierFlag);
-		if(!bAcceptedAccessSpecifier)
+		if (!bAcceptedAccessSpecifier)
 		{
 			MessageLog.Warning(*LOCTEXT("WrongAccessSpecifier_Error", "Wrong access specifier @@").ToString(), Context.EntryPoint);
 		}
@@ -1636,7 +1627,7 @@ void FKismetCompilerContext::PrecompileFunction(FKismetFunctionContext& Context)
 		Context.Function->StaticLink(true);
 
 		// Make sure the function signature is valid if this is an override
-		if (ParentFunction)
+		if (ParentFunction != nullptr)
 		{
 			// Verify the signature
 			if (!ParentFunction->IsSignatureCompatibleWith(Context.Function))
@@ -1650,7 +1641,7 @@ void FKismetCompilerContext::PrecompileFunction(FKismetFunctionContext& Context)
 			}
 			const bool bEmptyCase = (0 == AccessSpecifierFlag);
 			const bool bDifferentAccessSpecifiers = AccessSpecifierFlag != (ParentFunction->FunctionFlags & FUNC_AccessSpecifiers);
-			if( !bEmptyCase && bDifferentAccessSpecifiers )
+			if (!bEmptyCase && bDifferentAccessSpecifiers)
 			{
 				MessageLog.Warning(*LOCTEXT("IncompatibleAccessSpecifier_Error", "Access specifier is not compatible the parent function @@").ToString(), Context.EntryPoint);
 			}
@@ -1670,11 +1661,10 @@ void FKismetCompilerContext::PrecompileFunction(FKismetFunctionContext& Context)
 
 		////////////////////////////////////////
 
-		if(Context.IsDelegateSignature())
+		if (Context.IsDelegateSignature())
 		{
 			Context.Function->FunctionFlags |= FUNC_Delegate;
-			UMulticastDelegateProperty* Property = FindObject<UMulticastDelegateProperty>(Context.NewClass, *Context.DelegateSignatureName.ToString());
-			if(Property)
+			if(UMulticastDelegateProperty* Property = FindObject<UMulticastDelegateProperty>(Context.NewClass, *Context.DelegateSignatureName.ToString()))
 			{
 				Property->SignatureFunction = Context.Function;
 			}
@@ -2938,7 +2928,9 @@ void FKismetCompilerContext::CreateAndProcessUbergraph()
 		{
 			if (OldEventGraph)
 			{
-				OldEventGraph->Rename(NULL, GetTransientPackage(), Blueprint->bIsRegeneratingOnLoad ? REN_ForceNoResetLoaders : 0);
+				// when doing the GMinimalCompileOnLoad path we need to make sure that loaders aren't reset.. we may 
+				// need to set bIsRegeneratingOnLoad but it's likely that the reinstancer won't comply
+				OldEventGraph->Rename(NULL, GetTransientPackage(), (Blueprint->bIsRegeneratingOnLoad || GMinimalCompileOnLoad) ? REN_ForceNoResetLoaders : 0);
 			}
 		}
 	}
@@ -3514,7 +3506,9 @@ void FKismetCompilerContext::Compile()
 			}
 		}
  
-		ERenameFlags RenameFlags = (REN_DontCreateRedirectors|REN_NonTransactional|(Blueprint->bIsRegeneratingOnLoad ? REN_ForceNoResetLoaders : 0));
+		// when doing the GMinimalCompileOnLoad path we need to make sure that loaders aren't reset.. we may 
+		// need to set bIsRegeneratingOnLoad but it's likely that the reinstancer won't comply
+		ERenameFlags RenameFlags = (REN_DontCreateRedirectors|REN_NonTransactional|((Blueprint->bIsRegeneratingOnLoad || GMinimalCompileOnLoad )? REN_ForceNoResetLoaders : 0));
 		TargetClass->Rename(*NewName, NULL, RenameFlags);
 	}
 	// <<< End Backwards Compatibility
diff --git a/Engine/Source/Editor/KismetCompiler/Private/KismetCompilerMisc.cpp b/Engine/Source/Editor/KismetCompiler/Private/KismetCompilerMisc.cpp
--- a/Engine/Source/Editor/KismetCompiler/Private/KismetCompilerMisc.cpp
+++ b/Engine/Source/Editor/KismetCompiler/Private/KismetCompilerMisc.cpp
@@ -26,6 +26,7 @@
 #include "K2Node_FunctionEntry.h"
 #include "K2Node_FunctionResult.h"
 #include "K2Node_Timeline.h"
+#include "K2Node_Variable.h"
 #include "KismetCompiledFunctionContext.h"
 #include "KismetCompiler.h"
 
@@ -395,8 +396,8 @@ bool FKismetCompilerUtilities::IsTypeCompatibleWithProperty(UEdGraphPin* SourceP
 			SourcePin);
 	}
 
-	// Now check the direction
-	if (Property->HasAnyPropertyFlags(CPF_Parm))
+	// Now check the direction if it is parameter coming in or out of a function call style node (variable nodes are excluded since they maybe local parameters)
+	if (Property->HasAnyPropertyFlags(CPF_Parm) && !SourcePin->GetOwningNode()->IsA(UK2Node_Variable::StaticClass()))
 	{
 		// Parameters are directional
 		const bool bOutParam = Property->HasAllPropertyFlags(CPF_ReturnParm) || (Property->HasAllPropertyFlags(CPF_OutParm) && !Property->HasAnyPropertyFlags(CPF_ReferenceParm));
@@ -1066,7 +1067,9 @@ UProperty* FKismetCompilerUtilities::CreatePropertyOnScope(UStruct* Scope, const
 	{
 		if (UObject* ExistingObject = CheckPropertyNameOnScope(Scope, PropertyName))
 		{
-			MessageLog.Error(*FString::Printf(TEXT("Internal Compiler Error:  Tried to create a property %s in scope %s, but %s already exists there."), *PropertyName.ToString(), (Scope ? *Scope->GetName() : TEXT("None")), *ExistingObject->GetFullName()));
+			const FString ScopeName((Scope != nullptr) ? Scope->GetName() : FString(TEXT("None")));
+			const FString ExistingTypeAndPath(ExistingObject->GetFullName(Scope));
+			MessageLog.Error(*FString::Printf(TEXT("Internal Compiler Error: Tried to create a property %s in scope %s, but another object (%s) already already exists there."), *PropertyName.ToString(), *ScopeName, *ExistingTypeAndPath));
 
 			// Find a free name, so we can still create the property to make it easier to spot the duplicates, and avoid crashing
 			uint32 Counter = 0;
diff --git a/Engine/Source/Editor/KismetCompiler/Private/KismetCompilerModule.cpp b/Engine/Source/Editor/KismetCompiler/Private/KismetCompilerModule.cpp
--- a/Engine/Source/Editor/KismetCompiler/Private/KismetCompilerModule.cpp
+++ b/Engine/Source/Editor/KismetCompiler/Private/KismetCompilerModule.cpp
@@ -43,6 +43,7 @@ class FKismet2CompilerModule : public IKismetCompilerInterface
 public:
 	// Implementation of the IKismetCompilerInterface
 	virtual void CompileBlueprint(class UBlueprint* Blueprint, const FKismetCompilerOptions& CompileOptions, FCompilerResultsLog& Results, TSharedPtr<class FBlueprintCompileReinstancer> ParentReinstancer = NULL, TArray<UObject*>* ObjLoaded = NULL) override;
+	virtual void RefreshVariables(UBlueprint* Blueprint) override final;
 	virtual void CompileStructure(class UUserDefinedStruct* Struct, FCompilerResultsLog& Results) override;
 	virtual void RecoverCorruptedBlueprint(class UBlueprint* Blueprint) override;
 	virtual void RemoveBlueprintGeneratedClasses(class UBlueprint* Blueprint) override;
@@ -308,6 +309,37 @@ void FKismet2CompilerModule::CompileBlueprint(class UBlueprint* Blueprint, const
 	}
 }
 
+void FKismet2CompilerModule::RefreshVariables(UBlueprint* Blueprint)
+{
+	const UEdGraphSchema_K2* K2Schema = GetDefault<UEdGraphSchema_K2>();
+	FCompilerResultsLog MessageLog;
+
+	for (int32 VarIndex = 0; VarIndex < Blueprint->NewVariables.Num(); ++VarIndex)
+	{
+		UProperty* ExistingVariable = Blueprint->GeneratedClass->FindPropertyByName(Blueprint->NewVariables[VarIndex].VarName);
+		if( ExistingVariable == nullptr || ExistingVariable->GetOuter() != Blueprint->GeneratedClass)
+		{
+			// this variable will need to be added, we can link it in where ever as the sort will fix everything:
+			UProperty* NewProperty = FKismetCompilerUtilities::CreatePropertyOnScope(
+				Blueprint->GeneratedClass, 
+				Blueprint->NewVariables[VarIndex].VarName, 
+				Blueprint->NewVariables[VarIndex].VarType, 
+				Blueprint->GeneratedClass, 
+				0,
+				K2Schema,
+				MessageLog);
+
+			NewProperty->PropertyLinkNext = Blueprint->GeneratedClass->PropertyLink;
+			Blueprint->GeneratedClass->PropertyLink = NewProperty;
+		}
+	}
+
+	// we can't reliably remove properties because there are so 
+	// many nodes that create properties (e.g. anim blueprints, ubergraphframe)
+	// so just leave them.. we only need to prevent data loss
+	// Also note that property order is going to be insane
+}
+
 // Recover a corrupted blueprint
 void FKismet2CompilerModule::RecoverCorruptedBlueprint(class UBlueprint* Blueprint)
 {
diff --git a/Engine/Source/Editor/KismetCompiler/Public/KismetCompiler.h b/Engine/Source/Editor/KismetCompiler/Public/KismetCompiler.h
--- a/Engine/Source/Editor/KismetCompiler/Public/KismetCompiler.h
+++ b/Engine/Source/Editor/KismetCompiler/Public/KismetCompiler.h
@@ -283,11 +283,14 @@ protected:
 	/** Creates a property with flags including PropertyFlags in the Scope structure for each entry in the Terms array */
 	void CreatePropertiesFromList(UStruct* Scope, UField**& PropertyStorageLocation, TIndirectArray<FBPTerminal>& Terms, uint64 PropertyFlags, bool bPropertiesAreLocal, bool bPropertiesAreParameters = false);
 
-	/** Creates the properties on a function that store the function parameters, results, and local variables */
-	void CreateLocalVariablesForFunction(FKismetFunctionContext& Context, UFunction* ParameterSignature);
+	/** Create the properties on a function for input/output parameters */
+	void CreateParametersForFunction(FKismetFunctionContext& Context, UFunction* ParameterSignature, UField**& FunctionPropertyStorageLocation);
+
+	/** Creates the properties on a function that store the local and event graph (if applicable) variables */
+	void CreateLocalVariablesForFunction(FKismetFunctionContext& Context, UField**& FunctionPropertyStorageLocation);
 
 	/** Creates user defined local variables for function */
-	void CreateUserDefinedLocalVariablesForFunction(FKismetFunctionContext& Context, UField**& PropertyStorageLocation);
+	void CreateUserDefinedLocalVariablesForFunction(FKismetFunctionContext& Context, UField**& FunctionPropertyStorageLocation);
 
 	/** Adds a default value entry into the DefaultPropertyValueMap for the property specified */
 	void SetPropertyDefaultValue(const UProperty* PropertyToSet, FString& Value);
diff --git a/Engine/Source/Editor/KismetCompiler/Public/KismetCompilerMisc.h b/Engine/Source/Editor/KismetCompiler/Public/KismetCompilerMisc.h
--- a/Engine/Source/Editor/KismetCompiler/Public/KismetCompilerMisc.h
+++ b/Engine/Source/Editor/KismetCompiler/Public/KismetCompilerMisc.h
@@ -176,6 +176,12 @@ public:
 
 	virtual void RegisterNets(FKismetFunctionContext& Context, UEdGraphNode* Node);
 
+	// Returns true if this kind of node needs to be handled in the first pass, prior to execution scheduling (this is only used for function entry and return nodes)
+	virtual bool RequiresRegisterNetsBeforeScheduling() const
+	{
+		return false;
+	}
+
 	/**
 	 * Creates a sanitized name.
 	 *
diff --git a/Engine/Source/Editor/KismetCompiler/Public/KismetCompilerModule.h b/Engine/Source/Editor/KismetCompiler/Public/KismetCompilerModule.h
--- a/Engine/Source/Editor/KismetCompiler/Public/KismetCompilerModule.h
+++ b/Engine/Source/Editor/KismetCompiler/Public/KismetCompilerModule.h
@@ -58,6 +58,14 @@ public:
 	 */
 	virtual void CompileBlueprint(class UBlueprint* Blueprint, const FKismetCompilerOptions& CompileOptions, FCompilerResultsLog& Results, TSharedPtr<class FBlueprintCompileReinstancer> ParentReinstancer = NULL, TArray<UObject*>* ObjLoaded = NULL)=0;
 
+	/**
+	 * Synchronizes Blueprint's GeneratedClass's properties with the NewVariable declarations in the blueprint
+	 * Used on load to ensure that all properties are present for instances.
+	 * 
+	 * @param Blueprint The blueprint that may be missing variables
+	 */
+	virtual void RefreshVariables(UBlueprint* Blueprint)=0;
+
 	/**
 	 * Compiles a user defined structure.
 	 *
diff --git a/Engine/Source/Editor/PropertyEditor/Private/PropertyNode.cpp b/Engine/Source/Editor/PropertyEditor/Private/PropertyNode.cpp
--- a/Engine/Source/Editor/PropertyEditor/Private/PropertyNode.cpp
+++ b/Engine/Source/Editor/PropertyEditor/Private/PropertyNode.cpp
@@ -2419,6 +2419,8 @@ void FPropertyNode::NotifyPostChange( FPropertyChangedEvent& InPropertyChangedEv
 		{
 			PropertyChain->SetActiveMemberPropertyNode( OriginalActiveProperty );
 			PropertyChain->SetActivePropertyNode( InPropertyChangedEvent.Property);
+		
+			InPropertyChangedEvent.SetActiveMemberProperty(OriginalActiveProperty);
 			InNotifyHook->NotifyPostChange( InPropertyChangedEvent, &PropertyChain.Get() );
 		}
 	}
diff --git a/Engine/Source/Editor/UnrealEd/Classes/ThumbnailRendering/BlueprintThumbnailRenderer.h b/Engine/Source/Editor/UnrealEd/Classes/ThumbnailRendering/BlueprintThumbnailRenderer.h
--- a/Engine/Source/Editor/UnrealEd/Classes/ThumbnailRendering/BlueprintThumbnailRenderer.h
+++ b/Engine/Source/Editor/UnrealEd/Classes/ThumbnailRendering/BlueprintThumbnailRenderer.h
@@ -22,7 +22,7 @@ class UBlueprintThumbnailRenderer : public UDefaultSizedThumbnailRenderer
 	GENERATED_UCLASS_BODY()
 
 	// Begin UThumbnailRenderer Object
-	virtual bool CanVisualizeAsset(UObject* Object) override;
+	UNREALED_API virtual bool CanVisualizeAsset(UObject* Object) override;
 	UNREALED_API virtual void Draw(UObject* Object, int32 X, int32 Y, uint32 Width, uint32 Height, FRenderTarget*, FCanvas* Canvas) override;
 	// End UThumbnailRenderer Object
 
diff --git a/Engine/Source/Editor/UnrealEd/Private/Kismet2/BlueprintEditorUtils.cpp b/Engine/Source/Editor/UnrealEd/Private/Kismet2/BlueprintEditorUtils.cpp
--- a/Engine/Source/Editor/UnrealEd/Private/Kismet2/BlueprintEditorUtils.cpp
+++ b/Engine/Source/Editor/UnrealEd/Private/Kismet2/BlueprintEditorUtils.cpp
@@ -113,6 +113,8 @@
 #include "Styling/SlateIconFinder.h"
 #define LOCTEXT_NAMESPACE "Blueprint"
 
+extern COREUOBJECT_API bool GMinimalCompileOnLoad;
+
 DEFINE_LOG_CATEGORY(LogBlueprintDebug);
 
 DEFINE_STAT(EKismetCompilerStats_NotifyBlueprintChanged);
@@ -367,6 +369,8 @@ void FBasePinChangeHelper::Broadcast(UBlueprint* InBlueprint, UK2Node_EditablePi
 			? Func->GetOwnerClass()
 			: (UClass*)(FunctionDefNode ? FunctionDefNode->SignatureClass : nullptr);
 
+		const bool bIsInterface = FBlueprintEditorUtils::IsInterfaceBlueprint(InBlueprint);
+
 		// Reconstruct all function call sites that call this function (in open blueprints)
 		for (TObjectIterator<UK2Node_CallFunction> It(RF_Transient); It; ++It)
 		{
@@ -381,17 +385,31 @@ void FBasePinChangeHelper::Broadcast(UBlueprint* InBlueprint, UK2Node_EditablePi
 					// been removed by the user (and moved off the Blueprint)
 					continue;
 				}
-
-				const bool bNameMatches = (CallSite->FunctionReference.GetMemberName() == FuncName);
-				const UClass* MemberParentClass = CallSite->FunctionReference.GetMemberParentClass(CallSite->GetBlueprintClassFromNode());
-				const bool bClassMatchesEasy = (MemberParentClass != nullptr) && (MemberParentClass->IsChildOf(SignatureClass) || MemberParentClass->IsChildOf(InBlueprint->GeneratedClass));
-				const bool bClassMatchesHard = (CallSite->FunctionReference.IsSelfContext()) && (SignatureClass == nullptr) 
-					&& (CallSiteBlueprint == InBlueprint || CallSiteBlueprint->SkeletonGeneratedClass->IsChildOf(InBlueprint->SkeletonGeneratedClass));
+				
 				const bool bValidSchema = CallSite->GetSchema() != nullptr;
-
-				if (bNameMatches && bValidSchema && (bClassMatchesEasy || bClassMatchesHard))
+				const bool bNameMatches = (CallSite->FunctionReference.GetMemberName() == FuncName);
+				if (bNameMatches && bValidSchema)
 				{
-					EditCallSite(CallSite, CallSiteBlueprint);
+					if (bIsInterface)
+					{
+						if (FBlueprintEditorUtils::FindFunctionInImplementedInterfaces(CallSiteBlueprint, FuncName))
+						{
+							EditCallSite(CallSite, CallSiteBlueprint);
+						}
+					}
+					else
+					{
+						const UClass* MemberParentClass = CallSite->FunctionReference.GetMemberParentClass(CallSite->GetBlueprintClassFromNode());
+						const bool bClassMatchesEasy = (MemberParentClass != nullptr)
+							&& (MemberParentClass->IsChildOf(SignatureClass) || MemberParentClass->IsChildOf(InBlueprint->GeneratedClass));
+						const bool bClassMatchesHard = !bClassMatchesEasy && CallSite->FunctionReference.IsSelfContext() && (SignatureClass == nullptr)
+							&& (CallSiteBlueprint == InBlueprint || CallSiteBlueprint->SkeletonGeneratedClass->IsChildOf(InBlueprint->SkeletonGeneratedClass));
+
+						if (bClassMatchesEasy || bClassMatchesHard)
+						{
+							EditCallSite(CallSite, CallSiteBlueprint);
+						}
+					}
 				}
 			}
 		}
@@ -1270,6 +1288,12 @@ void FBlueprintEditorUtils::ForceLoadMembers(UObject* Object)
 	FRegenerationHelper::ForcedLoadMembers(Object);
 }
 
+void FBlueprintEditorUtils::RefreshVariables(UBlueprint* Blueprint)
+{
+	// module punchthrough:
+	IKismetCompilerInterface& Compiler = FModuleManager::LoadModuleChecked<IKismetCompilerInterface>(KISMET_COMPILER_MODULENAME);
+	Compiler.RefreshVariables(Blueprint);
+}
 
 UClass* FBlueprintEditorUtils::RegenerateBlueprintClass(UBlueprint* Blueprint, UClass* ClassToRegenerate, UObject* PreviousCDO, TArray<UObject*>& ObjLoaded)
 {
@@ -1314,21 +1338,27 @@ UClass* FBlueprintEditorUtils::RegenerateBlueprintClass(UBlueprint* Blueprint, U
 		int32 OldSkelLinkerIdx = INDEX_NONE;
 		int32 OldGenLinkerIdx = INDEX_NONE;
 		FLinkerLoad* OldLinker = Blueprint->GetLinker();
-		for( int32 i = 0; i < OldLinker->ExportMap.Num(); i++ )
+		// this linker related logic is not needed when running the 
+		// GMinimalCompileOnLoad pass, althoug hit is somewhat distressing
+		// that the linker is missing sometimes. Needs more investigation.
+		if(!GMinimalCompileOnLoad)
 		{
-			FObjectExport& ThisExport = OldLinker->ExportMap[i];
-			if( ThisExport.ObjectName == SkeletonName )
+			for( int32 i = 0; i < OldLinker->ExportMap.Num(); i++ )
 			{
-				OldSkelLinkerIdx = i;
-			}
-			else if( ThisExport.ObjectName == GeneratedName )
-			{
-				OldGenLinkerIdx = i;
-			}
+				FObjectExport& ThisExport = OldLinker->ExportMap[i];
+				if( ThisExport.ObjectName == SkeletonName )
+				{
+					OldSkelLinkerIdx = i;
+				}
+				else if( ThisExport.ObjectName == GeneratedName )
+				{
+					OldGenLinkerIdx = i;
+				}
 
-			if( OldSkelLinkerIdx != INDEX_NONE && OldGenLinkerIdx != INDEX_NONE )
-			{
-				break;
+				if( OldSkelLinkerIdx != INDEX_NONE && OldGenLinkerIdx != INDEX_NONE )
+				{
+					break;
+				}
 			}
 		}
 
@@ -1382,8 +1412,13 @@ UClass* FBlueprintEditorUtils::RegenerateBlueprintClass(UBlueprint* Blueprint, U
 			// that nodes have a chance to create all the pins they'll expect when they compile.
 			// A good example of why this is necessary is UK2Node_BaseAsyncTask::AllocateDefaultPins
 			// and it's companion function UK2Node_BaseAsyncTask::ExpandNode.
-			FBlueprintEditorUtils::ReconstructAllNodes(Blueprint);
-
+			// When running the GMinimalCompileOnLoad pass we will run node reconstruction explicitly
+			// in the compile manager:
+			if(!GMinimalCompileOnLoad)
+			{
+				FBlueprintEditorUtils::ReconstructAllNodes(Blueprint);
+			}
+			
 			FBlueprintEditorUtils::ReplaceDeprecatedNodes(Blueprint);
 
 			// Compile the actual blueprint
@@ -2058,6 +2093,8 @@ void FBlueprintEditorUtils::MarkBlueprintAsStructurallyModified(UBlueprint* Blue
 
 						SkelBlueprint->BroadcastCompiled();
 
+						SkelBlueprint->MarkPackageDirty();
+
 						SkeletalRecompileChildren(ChildrenOfClass, bIsCompilingOnLoad);
 						SkelBlueprint->bIsRegeneratingOnLoad = bWasRegenerating;
 					}
@@ -2277,12 +2314,12 @@ UEdGraph* FBlueprintEditorUtils::CreateNewGraph(UObject* ParentScope, const FNam
 	return NewGraph;
 }
 
-UFunction* FBlueprintEditorUtils::FindFunctionInImplementedInterfaces(const UBlueprint* Blueprint, const FName& FunctionName, bool * bOutInvalidInterface )
+UFunction* FBlueprintEditorUtils::FindFunctionInImplementedInterfaces(const UBlueprint* Blueprint, const FName& FunctionName, bool * bOutInvalidInterface, bool bGetAllInterfaces)
 {
 	if(Blueprint)
 	{
 		TArray<UClass*> InterfaceClasses;
-		FindImplementedInterfaces(Blueprint, false, InterfaceClasses);
+		FindImplementedInterfaces(Blueprint, bGetAllInterfaces, InterfaceClasses);
 
 		if( bOutInvalidInterface )
 		{
@@ -2294,6 +2331,13 @@ UFunction* FBlueprintEditorUtils::FindFunctionInImplementedInterfaces(const UBlu
 		{
 			if( SearchClass )
 			{
+				// Use the skeleton class if possible, as the generated class may not always be up-to-date (e.g. if the compile state is dirty).
+				UBlueprint* InterfaceBlueprint = Cast<UBlueprint>(SearchClass->ClassGeneratedBy);
+				if (InterfaceBlueprint && InterfaceBlueprint->SkeletonGeneratedClass)
+				{
+					SearchClass = InterfaceBlueprint->SkeletonGeneratedClass;
+				}
+
 				do
 				{
 					if( UFunction* OverriddenFunction = SearchClass->FindFunctionByName(FunctionName, EIncludeSuperFlag::ExcludeSuper) )
@@ -2582,17 +2626,29 @@ void FBlueprintEditorUtils::RenameGraph(UEdGraph* Graph, const FString& NewNameS
 				}
 			}
 		}
-
-		// Rename any function call points
-		TArray<UK2Node_CallFunction*> AllFunctionCalls;
-		FBlueprintEditorUtils::GetAllNodesOfClass<UK2Node_CallFunction>(Blueprint, AllFunctionCalls);
 	
-		for (UK2Node_CallFunction* FunctionNode : AllFunctionCalls)
+		// Rename any function call points
+		for (TObjectIterator<UK2Node_CallFunction> It(RF_Transient); It; ++It)
 		{
-			if (FunctionNode->FunctionReference.IsSelfContext() && (FunctionNode->FunctionReference.GetMemberName() == OldGraphName))
+			UK2Node_CallFunction* FunctionNode = *It;
+			if (FunctionNode && !FunctionNode->HasAnyFlags(RF_Transient) && Cast<UEdGraph>(FunctionNode->GetOuter()))
 			{
-				FunctionNode->Modify();
-				FunctionNode->FunctionReference.SetSelfMember(NewName);
+				if (UBlueprint* FunctionNodeBlueprint = FBlueprintEditorUtils::FindBlueprintForNode(FunctionNode))
+				{
+					if (FunctionNode->FunctionReference.GetMemberName() == OldGraphName)
+					{
+						if (FunctionNode->FunctionReference.IsSelfContext())
+						{
+							FunctionNode->Modify();
+							FunctionNode->FunctionReference.SetSelfMember(NewName);
+						}
+						else if (FunctionNode->FunctionReference.GetMemberParentClass() == Blueprint->GeneratedClass)
+						{
+							FunctionNode->Modify();
+							FunctionNode->FunctionReference.SetExternalMember(NewName, Blueprint->GeneratedClass);
+						}
+					}
+				}
 			}
 		}
 
@@ -4811,7 +4867,7 @@ FBPVariableDescription* FBlueprintEditorUtils::FindLocalVariable(UBlueprint* InB
 FBPVariableDescription* FBlueprintEditorUtils::FindLocalVariable(const UBlueprint* InBlueprint, const UEdGraph* InScopeGraph, const FName InVariableName, class UK2Node_FunctionEntry** OutFunctionEntry)
 {
 	FBPVariableDescription* ReturnVariable = nullptr;
-	if(DoesSupportLocalVariables(InScopeGraph))
+	if (DoesSupportLocalVariables(InScopeGraph))
 	{
 		UEdGraph* FunctionGraph = GetTopLevelGraph(InScopeGraph);
 		TArray<UK2Node_FunctionEntry*> GraphNodes;
@@ -4819,18 +4875,22 @@ FBPVariableDescription* FBlueprintEditorUtils::FindLocalVariable(const UBlueprin
 
 		bool bFoundLocalVariable = false;
 
-		// There is only ever 1 function entry
-		check(GraphNodes.Num() == 1)
-		for( int32 VarIdx = 0; VarIdx < GraphNodes[0]->LocalVariables.Num(); ++VarIdx )
+		if (GraphNodes.Num() > 0)
 		{
-			if(GraphNodes[0]->LocalVariables[VarIdx].VarName == InVariableName)
+			// If there is an entry node, there should only be one
+			check(GraphNodes.Num() == 1);
+
+			for (int32 VarIdx = 0; VarIdx < GraphNodes[0]->LocalVariables.Num(); ++VarIdx)
 			{
-				if (OutFunctionEntry)
+				if (GraphNodes[0]->LocalVariables[VarIdx].VarName == InVariableName)
 				{
-					*OutFunctionEntry = GraphNodes[0];
+					if (OutFunctionEntry)
+					{
+						*OutFunctionEntry = GraphNodes[0];
+					}
+					ReturnVariable = &GraphNodes[0]->LocalVariables[VarIdx];
+					break;
 				}
-				ReturnVariable = &GraphNodes[0]->LocalVariables[VarIdx];
-				break;
 			}
 		}
 	}
@@ -6781,14 +6841,31 @@ bool FBlueprintEditorUtils::FixLevelScriptActorBindings(ALevelScriptActor* Level
 					}
 					else
 					{
-						UE_LOG(LogBlueprint, Error, TEXT("Unable to bind event for node: %s (Owner: %s)! Can't get FMulticastScriptDelegate Target Delegate"), *GetNameSafe(EventNode), *GetNameSafe(EventNode->EventOwner));
+						FString ErrorFormat = LOCTEXT("FailedLSABinding_NoDelegate", "Unable to bind event node %s. Target instance is missing the expected delegate - is it a Blueprint in need of compile?").ToString();
+						if (FCompilerResultsLog* ActiveLog = FCompilerResultsLog::GetEventTarget())
+						{
+							ActiveLog->Warning(*FString::Printf(*ErrorFormat, TEXT("@@")), EventNode);
+						}
+						else
+						{
+							// For some reason, we don't have a valid entry point for the event in the LSA...
+							UE_LOG(LogBlueprint, Warning, TEXT("%s"), *FString::Printf(*ErrorFormat, *EventNode->GetName()));
+						}
 						bWasSuccessful = false;
 					}
 				}
 				else
 				{
-					// For some reason, we don't have a valid entry point for the event in the LSA...
-					UE_LOG(LogBlueprint, Warning, TEXT("Unable to bind event for node %s!  Please recompile the level blueprint."), (EventNode ? *EventNode->GetName() : TEXT("unknown")));
+					FString ErrorFormat = LOCTEXT("FailedLSABinding_NoFunction", "Unable to bind event node %s. Please recompile the level Blueprint.").ToString();
+					if (FCompilerResultsLog* ActiveLog = FCompilerResultsLog::GetEventTarget())
+					{
+						ActiveLog->Warning(*FString::Printf(*ErrorFormat, TEXT("@@")), EventNode);
+					}
+					else
+					{
+						// For some reason, we don't have a valid entry point for the event in the LSA...
+						UE_LOG(LogBlueprint, Warning, TEXT("%s"), *FString::Printf(*ErrorFormat, (EventNode ? *EventNode->GetName() : TEXT("[unknown]"))));
+					}
 					bWasSuccessful = false;
 				}
 			}
diff --git a/Engine/Source/Editor/UnrealEd/Private/Kismet2/Kismet2.cpp b/Engine/Source/Editor/UnrealEd/Private/Kismet2/Kismet2.cpp
--- a/Engine/Source/Editor/UnrealEd/Private/Kismet2/Kismet2.cpp
+++ b/Engine/Source/Editor/UnrealEd/Private/Kismet2/Kismet2.cpp
@@ -749,6 +749,8 @@ bool FKismetEditorUtilities::IsReferencedByUndoBuffer(UBlueprint* Blueprint)
 	return (TotalReferenceCount > NonUndoReferenceCount);
 }
 
+extern COREUOBJECT_API bool GMinimalCompileOnLoad;
+
 void FKismetEditorUtilities::CompileBlueprint(UBlueprint* BlueprintObj, bool bIsRegeneratingOnLoad, bool bSkipGarbageCollection, bool bSaveIntermediateProducts, FCompilerResultsLog* pResults, bool bSkeletonUpToDate, bool bBatchCompile, bool bAddInstrumentation)
 {
 	FSecondsCounterScope Timer(BlueprintCompileAndLoadTimerData); 
@@ -839,6 +841,16 @@ void FKismetEditorUtilities::CompileBlueprint(UBlueprint* BlueprintObj, bool bIs
 
 	ReinstanceHelper->UpdateBytecodeReferences();
 
+	// in case any errors/warnings have been added since the call to Compiler.CompileBlueprint()
+	if (Results.NumErrors > 0)
+	{
+		BlueprintObj->Status = BS_Error;
+	}
+	else if (Results.NumWarnings > 0)
+	{
+		BlueprintObj->Status = BS_UpToDateWithWarnings;
+	}
+
 	const bool bIsInterface = FBlueprintEditorUtils::IsInterfaceBlueprint(BlueprintObj);
 	const bool bLetReinstancerRefreshDependBP = !bIsRegeneratingOnLoad && (OldClass != NULL) && !bIsInterface;
 	if (bLetReinstancerRefreshDependBP)
@@ -850,7 +862,8 @@ void FKismetEditorUtilities::CompileBlueprint(UBlueprint* BlueprintObj, bool bIs
 		ReinstanceHelper->ListDependentBlueprintsToRefresh(DependentBPs);
 	}
 
-	if (!bIsRegeneratingOnLoad && (OldClass != NULL))
+	// we always need to run reinstancing when using GMinimalCompileOnLoad pass:
+	if ( (!bIsRegeneratingOnLoad || GMinimalCompileOnLoad) && (OldClass != NULL))
 	{
 		// Strip off any external components from the CDO, if needed because of reparenting, etc
 		FKismetEditorUtilities::StripExternalComponents(BlueprintObj);
diff --git a/Engine/Source/Editor/UnrealEd/Private/Kismet2/KismetReinstanceUtilities.cpp b/Engine/Source/Editor/UnrealEd/Private/Kismet2/KismetReinstanceUtilities.cpp
--- a/Engine/Source/Editor/UnrealEd/Private/Kismet2/KismetReinstanceUtilities.cpp
+++ b/Engine/Source/Editor/UnrealEd/Private/Kismet2/KismetReinstanceUtilities.cpp
@@ -902,6 +902,21 @@ void FBlueprintCompileReinstancer::ReinstanceObjects(bool bForceAlwaysReinstance
 					FKismetEditorUtilities::RecompileBlueprintBytecode(BP, nullptr, true);
 					ensure(0 == DependentBlueprintsToRecompile.Num());
 					CompiledBlueprints.Add(BP);
+
+					// We won't re-instance (and thus won't finalize) objects of this type, but we still need to know if the current
+					// Actor selection in the level editor includes instances of this type, and force a refresh on the selection if so.
+					// Note: If this is a Component BP instead of an Actor BP, the refresh is handled by IWCE's OnCompiled() BP delegate.
+					if (GEditor && BP->ParentClass && BP->ParentClass->IsChildOf<AActor>())
+					{
+						for (FSelectionIterator It(GEditor->GetSelectedActorIterator()); It; ++It)
+						{
+							AActor* ActorInstance = CastChecked<AActor>(*It);
+							if (ActorInstance->IsA(BP->GeneratedClass))
+							{
+								GEditor->SelectActor(ActorInstance, /*bInSelected =*/true, /*bNotify =*/true, false, true);
+							}
+						}
+					}
 				}
 
 
diff --git a/Engine/Source/Editor/UnrealEd/Public/Kismet2/BlueprintEditorUtils.h b/Engine/Source/Editor/UnrealEd/Public/Kismet2/BlueprintEditorUtils.h
--- a/Engine/Source/Editor/UnrealEd/Public/Kismet2/BlueprintEditorUtils.h
+++ b/Engine/Source/Editor/UnrealEd/Public/Kismet2/BlueprintEditorUtils.h
@@ -178,6 +178,11 @@ public:
 	 */
 	static void ForceLoadMembers(UObject* Object);
 
+	/**
+	 *  Synchronizes Blueprint's GeneratedClass's properties with the NewVariable declarations in the blueprint
+	 */
+	static void RefreshVariables(UBlueprint* Blueprint);
+
 	/**
 	 * Regenerates the class at class load time, and refreshes the blueprint
 	 */
@@ -1143,7 +1148,7 @@ public:
 	static void UpdateStalePinWatches( UBlueprint* Blueprint );
 
 	/** Return the first function from implemented interface with given name */
-	static UFunction* FindFunctionInImplementedInterfaces(const UBlueprint* Blueprint, const FName& FunctionName, bool* bOutInvalidInterface = NULL );
+	static UFunction* FindFunctionInImplementedInterfaces(const UBlueprint* Blueprint, const FName& FunctionName, bool* bOutInvalidInterface = nullptr, bool bGetAllInterfaces = false);
 
 	/** 
 	 * Build a list of all interface classes either implemented by this blueprint or through inheritance
diff --git a/Engine/Source/Runtime/CoreUObject/Private/Blueprint/BlueprintSupport.cpp b/Engine/Source/Runtime/CoreUObject/Private/Blueprint/BlueprintSupport.cpp
--- a/Engine/Source/Runtime/CoreUObject/Private/Blueprint/BlueprintSupport.cpp
+++ b/Engine/Source/Runtime/CoreUObject/Private/Blueprint/BlueprintSupport.cpp
@@ -22,6 +22,11 @@
 
 DEFINE_LOG_CATEGORY_STATIC(LogBlueprintSupport, Log, All);
 
+// the GMinimalCompileOnLoad flag simply means that we don't do a full compile pass while loading
+// data. This enables us to only compile when all dependencies have also been loaded, increasing
+// reliability of compiler results. This loading path does not yet function correctly:
+COREUOBJECT_API bool GMinimalCompileOnLoad = false;
+
 /**
  * Defined in BlueprintSupport.cpp
  * Duplicates all fields of a class in depth-first order. It makes sure that everything contained
@@ -232,17 +237,21 @@ FScopedClassDependencyGather::~FScopedClassDependencyGather()
 			}
 		};
 
-		for( ; DependencyIter; ++DependencyIter )
+		if(!GMinimalCompileOnLoad)
 		{
-			UClass* Dependency = *DependencyIter;
-			if( Dependency->ClassGeneratedBy != BatchMasterClass->ClassGeneratedBy )
+			for( ; DependencyIter; ++DependencyIter )
 			{
-				RecompileClassLambda(Dependency);
+				UClass* Dependency = *DependencyIter;
+				if( Dependency->ClassGeneratedBy != BatchMasterClass->ClassGeneratedBy )
+				{
+					RecompileClassLambda(Dependency);
+				}
 			}
+
+			// Finally, recompile the master class to make sure it gets updated too
+			RecompileClassLambda(BatchMasterClass);
 		}
 
-		// Finally, recompile the master class to make sure it gets updated too
-		RecompileClassLambda(BatchMasterClass);
 		BatchMasterClass = NULL;
 	}
 }
@@ -375,6 +384,8 @@ bool FLinkerLoad::RegenerateBlueprintClass(UClass* LoadClass, UObject* ExportObj
 			UObject* BlueprintObject = Class->ClassGeneratedBy;
 			if (BlueprintObject && BlueprintObject->HasAnyFlags(RF_BeingRegenerated))
 			{
+				// This code appears to be completely unused:
+
 				// Always load the parent blueprint here in case there is a circular dependency. This will
 				// ensure that the blueprint is fully serialized before attempting to regenerate the class.
 				FPreloadMembersHelper::PreloadObject(BlueprintObject);
@@ -458,6 +469,52 @@ static PlaceholderType* MakeImportPlaceholder(UObject* Outer, const TCHAR* Targe
 	return PlaceholderObj;
 }
 
+/** Recursive utility function, set up to find a specific import that has already been created (emulates a block from FLinkerLoad::CreateImport)*/
+static UObject* FindExistingImportObject(const int32 Index, const TArray<FObjectImport>& ImportMap)
+{
+	const FObjectImport& Import = ImportMap[Index];
+
+	UObject* FindOuter = nullptr;
+	if (Import.OuterIndex.IsImport())
+	{
+		int32 OuterIndex = Import.OuterIndex.ToImport();
+		const FObjectImport& OuterImport = ImportMap[OuterIndex];
+
+		if (OuterImport.XObject != nullptr)
+		{
+			FindOuter = OuterImport.XObject;
+		}
+		else
+		{
+			FindOuter = FindExistingImportObject(OuterIndex, ImportMap);
+		}
+	}
+
+	UObject* FoundObject = nullptr;
+	if (FindOuter != nullptr || Import.OuterIndex.IsNull())
+	{
+		if (UObject* ClassPackage = FindObject<UPackage>(/*Outer =*/nullptr, *Import.ClassPackage.ToString()))
+		{
+			if (UClass* ImportClass = FindObject<UClass>(ClassPackage, *Import.ClassName.ToString()))
+			{
+				// This function is set up to emulate a block towards the top of 
+				// FLinkerLoad::CreateImport(). However, since this is used in 
+				// deferred dependency loading we need to be careful not to invoke
+				// subsequent loads. The block in CreateImport() calls Preload() 
+				// and GetDefaultObject() which are not suitable here, so to 
+				// emulate/keep the contract that that block provides, we'll only 
+				// lookup the object if its class is loaded, and has a CDO (this
+				// is just to mitigate risk from this change)
+				if (!ImportClass->HasAnyFlags(RF_NeedLoad) && ImportClass->ClassDefaultObject)
+				{
+					FoundObject = StaticFindObjectFast(ImportClass, FindOuter, Import.ObjectName);
+				}				
+			}
+		}
+	}
+	return FoundObject;
+}
+
 bool FLinkerLoad::DeferPotentialCircularImport(const int32 Index)
 {
 #if USE_CIRCULAR_DEPENDENCY_LOAD_DEFERRING
@@ -479,6 +536,22 @@ bool FLinkerLoad::DeferPotentialCircularImport(const int32 Index)
 
 	if ((LoadFlags & LOAD_DeferDependencyLoads) && !IsImportNative(Index))
 	{
+		// emulate the block in CreateImport(), that attempts to find an existing
+		// object in memory first... this is to account for async loading, which
+		// can clear Import.XObject (via FLinkerManager::DissociateImportsAndForcedExports)
+		// at inopportune times (after it's already been set) - in this case
+		// we shouldn't need a placeholder, because the object already exists; we
+		// just need to keep from serializing it any further (which is why we've
+		// emulated it here, to cut out on a Preload() call)
+		if (!GIsEditor && !IsRunningCommandlet())
+		{
+			Import.XObject = FindExistingImportObject(Index, ImportMap);
+			if (Import.XObject)
+			{
+				return true;
+			}
+		}
+
 		if (UObject* ClassPackage = FindObject<UPackage>(/*Outer =*/nullptr, *Import.ClassPackage.ToString()))
 		{
 			if (const UClass* ImportClass = FindObject<UClass>(ClassPackage, *Import.ClassName.ToString()))
@@ -1729,6 +1802,13 @@ void FLinkerLoad::CreateDynamicTypeLoader()
 	// Try to get dependencies for dynamic classes
 	TArray<FBlueprintDependencyData> DependencyData;
 	FConvertedBlueprintsDependencies::Get().GetAssets(LinkerRoot->GetFName(), DependencyData);
+	if (!IsEventDrivenLoaderEnabled())
+	{
+		DependencyData.RemoveAll([=](const FBlueprintDependencyData& InData) -> bool
+		{
+			return InData.ObjectRef.PackageName == LinkerRoot->GetFName();
+		});
+	}
 
 	const FName DynamicClassName = UDynamicClass::StaticClass()->GetFName();
 	const FName DynamicClassPackageName = UDynamicClass::StaticClass()->GetOuterUPackage()->GetFName();
@@ -1782,116 +1862,116 @@ void FLinkerLoad::CreateDynamicTypeLoader()
 
 	if (GEventDrivenLoaderEnabled)
 	{
-	const FString DynamicTypePath = GetExportPathName(DynamicTypeExportIndex);
-	const FName DynamicTypeClassName = GetDynamicTypeClassName(*DynamicTypePath);
-	ensure(DynamicTypeClassName != NAME_None);
-	const bool bIsDynamicClass = DynamicTypeClassName == DynamicClassName;
-	const bool bIsDynamicStruct = DynamicTypeClassName == UScriptStruct::StaticClass()->GetFName();
-
-	if(bIsDynamicClass || bIsDynamicStruct)
-	{
-		FObjectExport* const CDOExport = bIsDynamicClass ? (new (ExportMap)FObjectExport()) : nullptr;
-		if(CDOExport)
-		{
-			const FString CDOName = FString(DEFAULT_OBJECT_PREFIX) + DynamicTypeExport->ObjectName.ToString();
-			CDOExport->ObjectName = *CDOName;
-			CDOExport->ThisIndex = FPackageIndex::FromExport(ExportMap.Num() - 1);
-			CDOExport->DynamicType = FObjectExport::EDynamicType::ClassDefaultObject;
-			CDOExport->ObjectFlags |= RF_Public | RF_ClassDefaultObject; //? 
-			CDOExport->ClassIndex = DynamicTypeExport->ThisIndex;
-		}
+		const FString DynamicTypePath = GetExportPathName(DynamicTypeExportIndex);
+		const FName DynamicTypeClassName = GetDynamicTypeClassName(*DynamicTypePath);
+		ensure(DynamicTypeClassName != NAME_None);
+		const bool bIsDynamicClass = DynamicTypeClassName == DynamicClassName;
+		const bool bIsDynamicStruct = DynamicTypeClassName == UScriptStruct::StaticClass()->GetFName();
 
-		FObjectExport* const FakeExports[] = { DynamicTypeExport , CDOExport }; // must be sync'ed with FBlueprintDependencyData::DependencyTypes
-		int32 RunningIndex = 0;
-		for(int32 LocExportIndex = 0; LocExportIndex < (sizeof(FakeExports)/sizeof(FakeExports[0])); LocExportIndex++)
+		if(bIsDynamicClass || bIsDynamicStruct)
 		{
-			FObjectExport* const Export = FakeExports[LocExportIndex];
-			if (!Export)
+			FObjectExport* const CDOExport = bIsDynamicClass ? (new (ExportMap)FObjectExport()) : nullptr;
+			if(CDOExport)
 			{
-				continue;
+				const FString CDOName = FString(DEFAULT_OBJECT_PREFIX) + DynamicTypeExport->ObjectName.ToString();
+				CDOExport->ObjectName = *CDOName;
+				CDOExport->ThisIndex = FPackageIndex::FromExport(ExportMap.Num() - 1);
+				CDOExport->DynamicType = FObjectExport::EDynamicType::ClassDefaultObject;
+				CDOExport->ObjectFlags |= RF_Public | RF_ClassDefaultObject; //? 
+				CDOExport->ClassIndex = DynamicTypeExport->ThisIndex;
 			}
-			Export->FirstExportDependency = RunningIndex;
 
-			enum class EDependencyType : uint8
+			FObjectExport* const FakeExports[] = { DynamicTypeExport , CDOExport }; // must be sync'ed with FBlueprintDependencyData::DependencyTypes
+			int32 RunningIndex = 0;
+			for(int32 LocExportIndex = 0; LocExportIndex < (sizeof(FakeExports)/sizeof(FakeExports[0])); LocExportIndex++)
 			{
-				SerializationBeforeSerialization,
-				CreateBeforeSerialization,
-				SerializationBeforeCreate,
-				CreateBeforeCreate,
-			};
+				FObjectExport* const Export = FakeExports[LocExportIndex];
+				if (!Export)
+				{
+					continue;
+				}
+				Export->FirstExportDependency = RunningIndex;
 
-			auto HandleDependencyTypeForExport = [&](EDependencyType InDependencyType)
-			{
-				for (int32 DependencyDataIndex = 0; DependencyDataIndex < DependencyData.Num(); DependencyDataIndex++)
+				enum class EDependencyType : uint8
 				{
-					const FBlueprintDependencyData& Import = DependencyData[DependencyDataIndex];
-					const FBlueprintDependencyType DependencyType = Import.DependencyTypes[LocExportIndex];
-					auto IsMatchingDependencyType = [](FBlueprintDependencyType InDependencyTypeStruct, EDependencyType InDependencyTypeLoc) -> bool
-					{
-						switch (InDependencyTypeLoc)
-						{
-						case EDependencyType::SerializationBeforeSerialization:
-							return InDependencyTypeStruct.bSerializationBeforeSerializationDependency;
-						case EDependencyType::CreateBeforeSerialization:
-							return InDependencyTypeStruct.bCreateBeforeSerializationDependency;
-						case EDependencyType::SerializationBeforeCreate:
-							return InDependencyTypeStruct.bSerializationBeforeCreateDependency;
-						case EDependencyType::CreateBeforeCreate:
-							return InDependencyTypeStruct.bCreateBeforeCreateDependency;
-						}
-						check(false);
-						return false;
-					};
-					if (IsMatchingDependencyType(DependencyType, InDependencyType))
+					SerializationBeforeSerialization,
+					CreateBeforeSerialization,
+					SerializationBeforeCreate,
+					CreateBeforeCreate,
+				};
+
+				auto HandleDependencyTypeForExport = [&](EDependencyType InDependencyType)
+				{
+					for (int32 DependencyDataIndex = 0; DependencyDataIndex < DependencyData.Num(); DependencyDataIndex++)
 					{
-						auto IncreaseDependencyTypeInExport = [](FObjectExport* InExport, EDependencyType InDependencyTypeLoc)
+						const FBlueprintDependencyData& Import = DependencyData[DependencyDataIndex];
+						const FBlueprintDependencyType DependencyType = Import.DependencyTypes[LocExportIndex];
+						auto IsMatchingDependencyType = [](FBlueprintDependencyType InDependencyTypeStruct, EDependencyType InDependencyTypeLoc) -> bool
 						{
-							check(InExport);
 							switch (InDependencyTypeLoc)
 							{
 							case EDependencyType::SerializationBeforeSerialization:
-								InExport->SerializationBeforeSerializationDependencies++;
-								break;
+								return InDependencyTypeStruct.bSerializationBeforeSerializationDependency;
 							case EDependencyType::CreateBeforeSerialization:
-								InExport->CreateBeforeSerializationDependencies++;
-								break;
+								return InDependencyTypeStruct.bCreateBeforeSerializationDependency;
 							case EDependencyType::SerializationBeforeCreate:
-								InExport->SerializationBeforeCreateDependencies++;
-								break;
+								return InDependencyTypeStruct.bSerializationBeforeCreateDependency;
 							case EDependencyType::CreateBeforeCreate:
-								InExport->CreateBeforeCreateDependencies++;
-								break;
+								return InDependencyTypeStruct.bCreateBeforeCreateDependency;
 							}
+							check(false);
+							return false;
 						};
-						IncreaseDependencyTypeInExport(Export, InDependencyType);
-
-						auto IndexInDependencyDataToImportIndex = [](int32 ArrayIndex) -> int32 { return ArrayIndex * 2; };
-						const int32 ImportIndex = IndexInDependencyDataToImportIndex(DependencyDataIndex);
-						PreloadDependencies.Add(FPackageIndex::FromImport(ImportIndex));
-						RunningIndex++;
+						if (IsMatchingDependencyType(DependencyType, InDependencyType))
+						{
+							auto IncreaseDependencyTypeInExport = [](FObjectExport* InExport, EDependencyType InDependencyTypeLoc)
+							{
+								check(InExport);
+								switch (InDependencyTypeLoc)
+								{
+								case EDependencyType::SerializationBeforeSerialization:
+									InExport->SerializationBeforeSerializationDependencies++;
+									break;
+								case EDependencyType::CreateBeforeSerialization:
+									InExport->CreateBeforeSerializationDependencies++;
+									break;
+								case EDependencyType::SerializationBeforeCreate:
+									InExport->SerializationBeforeCreateDependencies++;
+									break;
+								case EDependencyType::CreateBeforeCreate:
+									InExport->CreateBeforeCreateDependencies++;
+									break;
+								}
+							};
+							IncreaseDependencyTypeInExport(Export, InDependencyType);
+
+							auto IndexInDependencyDataToImportIndex = [](int32 ArrayIndex) -> int32 { return ArrayIndex * 2; };
+							const int32 ImportIndex = IndexInDependencyDataToImportIndex(DependencyDataIndex);
+							PreloadDependencies.Add(FPackageIndex::FromImport(ImportIndex));
+							RunningIndex++;
+						}
 					}
-				}
-			};
+				};
 
-			// the order of Packages in PreloadDependencie must match FAsyncPackage::SetupExports_Event
+				// the order of Packages in PreloadDependencie must match FAsyncPackage::SetupExports_Event
 
-			HandleDependencyTypeForExport(EDependencyType::SerializationBeforeSerialization);
-			HandleDependencyTypeForExport(EDependencyType::CreateBeforeSerialization);
+				HandleDependencyTypeForExport(EDependencyType::SerializationBeforeSerialization);
+				HandleDependencyTypeForExport(EDependencyType::CreateBeforeSerialization);
 
-			if (bIsDynamicClass && (Export == CDOExport))
-			{
-				// Add a serializebeforecreate arc from the class on the CDO. That will force us to finish the class before we create the CDO....
-				// and that will make sure that we load the class before we serialize things that reference the CDO.
-				Export->SerializationBeforeCreateDependencies++;
-				PreloadDependencies.Add(DynamicTypeExport->ThisIndex);
-				RunningIndex++;
-			}
+				if (bIsDynamicClass && (Export == CDOExport))
+				{
+					// Add a serializebeforecreate arc from the class on the CDO. That will force us to finish the class before we create the CDO....
+					// and that will make sure that we load the class before we serialize things that reference the CDO.
+					Export->SerializationBeforeCreateDependencies++;
+					PreloadDependencies.Add(DynamicTypeExport->ThisIndex);
+					RunningIndex++;
+				}
 
-			HandleDependencyTypeForExport(EDependencyType::SerializationBeforeCreate);
-			HandleDependencyTypeForExport(EDependencyType::CreateBeforeCreate);
+				HandleDependencyTypeForExport(EDependencyType::SerializationBeforeCreate);
+				HandleDependencyTypeForExport(EDependencyType::CreateBeforeCreate);
+			}
 		}
 	}
-	}
 
 	LinkerRoot->SetPackageFlags(LinkerRoot->GetPackageFlags() | PKG_CompiledIn);
 }
@@ -2317,6 +2397,22 @@ void FConvertedBlueprintsDependencies::FillUsedAssetsInDynamicClass(UDynamicClas
 	}
 }
 
+bool FBlueprintDependencyData::ContainsDependencyData(TArray<FBlueprintDependencyData>& Assets, int16 ObjectRefIndex)
+{
+	return nullptr != Assets.FindByPredicate([=](const FBlueprintDependencyData& Data) -> bool
+	{
+		return Data.ObjectRefIndex == ObjectRefIndex;
+	});
+};
+
+void FBlueprintDependencyData::AppendUniquely(TArray<FBlueprintDependencyData>& Destination, const TArray<FBlueprintDependencyData>& AdditionalData)
+{
+	for (const FBlueprintDependencyData& Data : AdditionalData)
+	{
+		Destination.AddUnique(Data);
+	}
+}
+
 
 #if WITH_EDITOR
 
diff --git a/Engine/Source/Runtime/CoreUObject/Private/Serialization/AsyncLoading.cpp b/Engine/Source/Runtime/CoreUObject/Private/Serialization/AsyncLoading.cpp
--- a/Engine/Source/Runtime/CoreUObject/Private/Serialization/AsyncLoading.cpp
+++ b/Engine/Source/Runtime/CoreUObject/Private/Serialization/AsyncLoading.cpp
@@ -1547,17 +1547,22 @@ FObjectImport* FAsyncPackage::FindExistingImport(int32 LocalImportIndex)
 				if (Import->XObject)
 				{
 //native blueprint 
+#if DO_CHECK
 					FName NameImportClass(Import->ClassName);
 					FName NameActualImportClass(Import->XObject->GetClass()->GetFName());
 					if (NameActualImportClass != NameImportClass)
 					{
 						static const FName NAME_BlueprintGeneratedClass("BlueprintGeneratedClass");
 						static const FName NAME_DynamicClass("DynamicClass");
-						bool bBroken = (NameImportClass != NAME_BlueprintGeneratedClass && NameImportClass != NAME_DynamicClass) ||
-							(NameActualImportClass != NAME_BlueprintGeneratedClass && NameActualImportClass != NAME_DynamicClass);
-						UE_CLOG(bBroken, LogStreaming, Fatal, TEXT("FAsyncPackage::FindExistingImport class mismatch %s != %s"), *NameActualImportClass.ToString(), *NameImportClass.ToString());
-					}
 
+						static const FName NAME_Function("Function");
+						static const FName NAME_DelegateFunction("DelegateFunction");
+
+						bool bSafeException = (NameImportClass == NAME_BlueprintGeneratedClass && NameActualImportClass == NAME_DynamicClass)
+							|| (NameImportClass == NAME_Function && NameActualImportClass == NAME_DelegateFunction);
+						UE_CLOG(!bSafeException, LogStreaming, Fatal, TEXT("FAsyncPackage::FindExistingImport class mismatch %s != %s"), *NameActualImportClass.ToString(), *NameImportClass.ToString());
+					}
+#endif //DO_CHECK
 					AddObjectReference(Import->XObject);
 				}
 			}
@@ -1765,7 +1770,7 @@ static bool IsFullyLoadedObj(UObject* Obj)
 		return true;
 	}
 
-	if (GEventDrivenLoaderEnabled)
+	if (GEventDrivenLoaderEnabled && EVENT_DRIVEN_ASYNC_LOAD_ACTIVE_AT_RUNTIME)
 	{
 		if (0 != (UD->ClassFlags & CLASS_Constructed))
 		{
@@ -5769,19 +5774,22 @@ void FAsyncPackage::CloseDelayedLinkers()
 	// Close any linkers that have been open as a result of blocking load while async loading
 	for (FLinkerLoad* LinkerToClose : DelayedLinkerClosePackages)
 	{
-		check(LinkerToClose);
-		check(LinkerToClose->LinkerRoot);
-		if (GEventDrivenLoaderEnabled)
+		if (LinkerToClose->LinkerRoot != nullptr)
 		{
-			FLinkerLoad* LinkerToReset = FLinkerLoad::FindExistingLinkerForPackage(CastChecked<UPackage>(LinkerToClose->LinkerRoot));
-			check(LinkerToReset == LinkerToClose);
-			if (LinkerToReset && LinkerToReset->AsyncRoot)
+			check(LinkerToClose);
+			//check(LinkerToClose->LinkerRoot);
+			if (GEventDrivenLoaderEnabled)
 			{
-				UE_LOG(LogStreaming, Error, TEXT("Linker cannot be reset right now...leaking %s"), *LinkerToReset->GetArchiveName());
-				continue;
+				FLinkerLoad* LinkerToReset = FLinkerLoad::FindExistingLinkerForPackage(CastChecked<UPackage>(LinkerToClose->LinkerRoot));
+				check(LinkerToReset == LinkerToClose);
+				if (LinkerToReset && LinkerToReset->AsyncRoot)
+				{
+					UE_LOG(LogStreaming, Error, TEXT("Linker cannot be reset right now...leaking %s"), *LinkerToReset->GetArchiveName());
+					continue;
+				}
 			}
+			FLinkerManager::Get().ResetLoaders(LinkerToClose->LinkerRoot);
 		}
-		FLinkerManager::Get().ResetLoaders(LinkerToClose->LinkerRoot);
 		check(LinkerToClose->LinkerRoot == nullptr);
 		check(LinkerToClose->AsyncRoot == nullptr);
 	}
diff --git a/Engine/Source/Runtime/CoreUObject/Private/UObject/ScriptCore.cpp b/Engine/Source/Runtime/CoreUObject/Private/UObject/ScriptCore.cpp
--- a/Engine/Source/Runtime/CoreUObject/Private/UObject/ScriptCore.cpp
+++ b/Engine/Source/Runtime/CoreUObject/Private/UObject/ScriptCore.cpp
@@ -880,12 +880,7 @@ void UObject::ProcessInternal( FFrame& Stack, RESULT_DECL )
 	{
 		if (!GIsReinstancing)
 		{
-			static int32 num = 0;
-			num++;
-			if (num < 5)
-			{
-				ensureMsgf(!GetClass()->HasAnyClassFlags(CLASS_NewerVersionExists), TEXT("Object '%s' is being used for execution, but its class is out of date and has been replaced with a recompiled class!"), *GetFullName());
-			}
+			ensureMsgf(!GetClass()->HasAnyClassFlags(CLASS_NewerVersionExists), TEXT("Object '%s' is being used for execution, but its class is out of date and has been replaced with a recompiled class!"), *GetFullName());
 		}
 		return;
 	}
diff --git a/Engine/Source/Runtime/CoreUObject/Public/Blueprint/BlueprintSupport.h b/Engine/Source/Runtime/CoreUObject/Public/Blueprint/BlueprintSupport.h
--- a/Engine/Source/Runtime/CoreUObject/Public/Blueprint/BlueprintSupport.h
+++ b/Engine/Source/Runtime/CoreUObject/Public/Blueprint/BlueprintSupport.h
@@ -290,14 +290,26 @@ struct COREUOBJECT_API FBlueprintDependencyData
 	// 1 - dependency type for CD0
 	FBlueprintDependencyType DependencyTypes[2];
 
+	int16 ObjectRefIndex; // NativizationWithoutEDLBT
+
 	FBlueprintDependencyData(const FBlueprintDependencyObjectRef& InObjectRef
 		, FBlueprintDependencyType InClassDependency
-		, FBlueprintDependencyType InCDODependency)
+		, FBlueprintDependencyType InCDODependency
+		, int16 InObjectRefIndex)
 		: ObjectRef(InObjectRef)
+		, ObjectRefIndex(InObjectRefIndex)
 	{
 		DependencyTypes[0] = InClassDependency;
 		DependencyTypes[1] = InCDODependency;
 	}
+
+	bool operator==(const FBlueprintDependencyData& Other) const
+	{
+		return Other.ObjectRefIndex == ObjectRefIndex;
+	}
+
+	static bool ContainsDependencyData(TArray<FBlueprintDependencyData>& Assets, int16 ObjectRefIndex);
+	static void AppendUniquely(TArray<FBlueprintDependencyData>& Destination, const TArray<FBlueprintDependencyData>& AdditionalData);
 };
 
 /**
diff --git a/Engine/Source/Runtime/CoreUObject/Public/UObject/Class.h b/Engine/Source/Runtime/CoreUObject/Public/UObject/Class.h
--- a/Engine/Source/Runtime/CoreUObject/Public/UObject/Class.h
+++ b/Engine/Source/Runtime/CoreUObject/Public/UObject/Class.h
@@ -1674,7 +1674,7 @@ public:
 	 *
 	 * @return The display name for this object.
 	 */
-	FText GetDisplayNameText(int32 NameIndex=INDEX_NONE) const;
+	virtual FText GetDisplayNameText(int32 NameIndex=INDEX_NONE) const;
 	FText GetDisplayNameTextByValue(int64 Value = INDEX_NONE) const;
 
 	/**
diff --git a/Engine/Source/Runtime/Engine/Classes/Engine/BlueprintGeneratedClass.h b/Engine/Source/Runtime/Engine/Classes/Engine/BlueprintGeneratedClass.h
--- a/Engine/Source/Runtime/Engine/Classes/Engine/BlueprintGeneratedClass.h
+++ b/Engine/Source/Runtime/Engine/Classes/Engine/BlueprintGeneratedClass.h
@@ -665,6 +665,10 @@ public:
 	UPROPERTY()
 	TMap<FName, struct FBlueprintCookedComponentInstancingData> CookedComponentInstancingData;
 
+	/** Flag used to indicate if this class has a nativized parent in a cooked build. */
+	UPROPERTY()
+	bool bHasNativizedParent;
+
 	/** 
 	 * Gets an array of all BPGeneratedClasses (including InClass as 0th element) parents of given generated class 
 	 *
@@ -683,6 +687,9 @@ public:
 	static void CreateComponentsForActor(const UClass* ThisClass, AActor* Actor);
 	static void CreateTimelineComponent(AActor* Actor, const UTimelineTemplate* TimelineTemplate);
 
+	/** Check for and handle manual application of default value overrides to instanced component subobjects that were inherited from a nativized parent class */
+	static void CheckAndApplyComponentTemplateOverrides(AActor* Actor);
+
 	// UObject interface
 	virtual void Serialize(FArchive& Ar) override;
 	virtual void PostLoad() override;
diff --git a/Engine/Source/Runtime/Engine/Classes/Engine/InheritableComponentHandler.h b/Engine/Source/Runtime/Engine/Classes/Engine/InheritableComponentHandler.h
--- a/Engine/Source/Runtime/Engine/Classes/Engine/InheritableComponentHandler.h
+++ b/Engine/Source/Runtime/Engine/Classes/Engine/InheritableComponentHandler.h
@@ -50,13 +50,13 @@ struct ENGINE_API FComponentKey
 	UActorComponent* GetOriginalTemplate() const;
 	bool RefreshVariableName();
 
-	class UBlueprintGeneratedClass* GetComponentOwner()  const { return OwnerClass; }
+	UClass* GetComponentOwner()  const { return OwnerClass; }
 	FName   GetSCSVariableName() const { return SCSVariableName; }
 	FGuid   GetAssociatedGuid()  const { return AssociatedGuid; }
 
 private: 
 	UPROPERTY()
-	class UBlueprintGeneratedClass* OwnerClass;
+	UClass* OwnerClass;
 
 	UPROPERTY()
 	FName SCSVariableName;
diff --git a/Engine/Source/Runtime/Engine/Classes/Engine/UserDefinedEnum.h b/Engine/Source/Runtime/Engine/Classes/Engine/UserDefinedEnum.h
--- a/Engine/Source/Runtime/Engine/Classes/Engine/UserDefinedEnum.h
+++ b/Engine/Source/Runtime/Engine/Classes/Engine/UserDefinedEnum.h
@@ -70,6 +70,8 @@ public:
 	virtual void PostEditUndo() override;
 	//~ End UObject Interface
 
+	virtual FText GetDisplayNameText(int32 NameIndex = INDEX_NONE) const;
+
 	FString GenerateNewEnumeratorName();
 #endif	// WITH_EDITOR
 };
diff --git a/Engine/Source/Runtime/Engine/Classes/Kismet/BlueprintSetLibrary.h b/Engine/Source/Runtime/Engine/Classes/Kismet/BlueprintSetLibrary.h
--- a/Engine/Source/Runtime/Engine/Classes/Kismet/BlueprintSetLibrary.h
+++ b/Engine/Source/Runtime/Engine/Classes/Kismet/BlueprintSetLibrary.h
@@ -102,7 +102,7 @@ public:
 	 * @param		B		Another set to intersect
 	 * @param		Result	Set to store results in
 	 */
-	UFUNCTION(BlueprintCallable, CustomThunk, meta=(DisplayName = "Intersection", CompactNodeTitle = "INTERSECTION", SetParam = "A,B,Result"), Category="Utilities|Set")
+	UFUNCTION(BlueprintCallable, CustomThunk, meta=(DisplayName = "Intersection", CompactNodeTitle = "INTERSECTION", SetParam = "A|B|Result"), Category="Utilities|Set")
 	static void Set_Intersection(const TSet<int32>& A, const TSet<int32>& B, TSet<int32>& Result );
 
 	/**
@@ -114,7 +114,7 @@ public:
 	 * @param		B		Another set to union
 	 * @param		Result	Set to store results in
 	 */
-	UFUNCTION(BlueprintCallable, CustomThunk, meta = (DisplayName = "Union", CompactNodeTitle = "UNION", SetParam = "A,B,Result"), Category = "Utilities|Set")
+	UFUNCTION(BlueprintCallable, CustomThunk, meta = (DisplayName = "Union", CompactNodeTitle = "UNION", SetParam = "A|B|Result"), Category = "Utilities|Set")
 	static void Set_Union(const TSet<int32>& A, const TSet<int32>& B, TSet<int32>& Result );
 
 	/**
@@ -127,7 +127,7 @@ public:
 	 * @param		B		Set of elements to remove from set A
 	 * @param		Result	Set containing all elements in A that are not found in B
 	 */
-	UFUNCTION(BlueprintCallable, CustomThunk, meta = (DisplayName = "Difference", CompactNodeTitle = "DIFFERENCE", SetParam = "A,B,Result"), Category = "Utilities|Set")
+	UFUNCTION(BlueprintCallable, CustomThunk, meta = (DisplayName = "Difference", CompactNodeTitle = "DIFFERENCE", SetParam = "A|B|Result"), Category = "Utilities|Set")
 	static void Set_Difference(const TSet<int32>& A, const TSet<int32>& B, TSet<int32>& Result );
 
 	DECLARE_FUNCTION(execSet_Add)
diff --git a/Engine/Source/Runtime/Engine/Classes/Kismet/KismetMathLibrary.h b/Engine/Source/Runtime/Engine/Classes/Kismet/KismetMathLibrary.h
--- a/Engine/Source/Runtime/Engine/Classes/Kismet/KismetMathLibrary.h
+++ b/Engine/Source/Runtime/Engine/Classes/Kismet/KismetMathLibrary.h
@@ -489,6 +489,10 @@ class ENGINE_API UKismetMathLibrary : public UBlueprintFunctionLibrary
 	UFUNCTION(BlueprintPure, Category="Math|Float", meta=(CompactNodeTitle = "e"))
 	static float Exp(float A);
 
+	/* Returns log of A base B (if B^R == A, returns R)*/
+	UFUNCTION(BlueprintPure, Category = "Math|Float")
+	static float Log(float A, float Base = 1.f);
+
 	/* Returns natural log of A (if e^R == A, returns R)*/
 	UFUNCTION(BlueprintPure, Category="Math|Float")
 	static float Loge(float A);
diff --git a/Engine/Source/Runtime/Engine/Private/Actor.cpp b/Engine/Source/Runtime/Engine/Private/Actor.cpp
--- a/Engine/Source/Runtime/Engine/Private/Actor.cpp
+++ b/Engine/Source/Runtime/Engine/Private/Actor.cpp
@@ -617,6 +617,17 @@ void AActor::PostLoadSubobjects(FObjectInstancingGraph* OuterInstanceGraph)
 
 	Super::PostLoadSubobjects(OuterInstanceGraph);
 
+	// If this is a Blueprint class, we may need to manually apply default value overrides to some inherited components in a cooked
+	// build scenario. This can occur, for example, if we have a nativized Blueprint class somewhere in the class inheritance hierarchy.
+	if (FPlatformProperties::RequiresCookedData())
+	{
+		const UBlueprintGeneratedClass* BPGC = Cast<UBlueprintGeneratedClass>(GetClass());
+		if (BPGC != nullptr && BPGC->bHasNativizedParent)
+		{
+			UBlueprintGeneratedClass::CheckAndApplyComponentTemplateOverrides(this);
+		}
+	}
+
 	ResetOwnedComponents();
 
 	if (RootComponent && bHadRoot && OldRoot != RootComponent)
@@ -2749,6 +2760,18 @@ void AActor::PostSpawnInitialize(FTransform const& UserSpawnTransform, AActor* I
 	// Call OnComponentCreated on all default (native) components
 	DispatchOnComponentsCreated(this);
 
+	// If this is a Blueprint class, we may need to manually apply default value overrides to some inherited components in a
+	// cooked build scenario. This can occur, for example, if we have a nativized Blueprint class in the inheritance hierarchy.
+	// Note: This should be done prior to executing the construction script, in case there are any dependencies on default values.
+	if (FPlatformProperties::RequiresCookedData())
+	{
+		const UBlueprintGeneratedClass* BPGC = Cast<UBlueprintGeneratedClass>(GetClass());
+		if (BPGC != nullptr && BPGC->bHasNativizedParent)
+		{
+			UBlueprintGeneratedClass::CheckAndApplyComponentTemplateOverrides(this);
+		}
+	}
+
 	// Register the actor's default (native) components, but only if we have a native scene root. If we don't, it implies that there could be only non-scene components
 	// at the native class level. In that case, if this is a Blueprint instance, we need to defer native registration until after SCS execution can establish a scene root.
 	// Note: This API will also call PostRegisterAllComponents() on the actor instance. If deferred, PostRegisterAllComponents() won't be called until the root is set by SCS.
diff --git a/Engine/Source/Runtime/Engine/Private/Blueprint.cpp b/Engine/Source/Runtime/Engine/Private/Blueprint.cpp
--- a/Engine/Source/Runtime/Engine/Private/Blueprint.cpp
+++ b/Engine/Source/Runtime/Engine/Private/Blueprint.cpp
@@ -17,6 +17,7 @@
 #include "Components/TimelineComponent.h"
 
 #if WITH_EDITOR
+#include "BlueprintCompilationManager.h"
 #include "Editor/UnrealEd/Classes/Settings/ProjectPackagingSettings.h"
 #include "Engine/SimpleConstructionScript.h"
 #include "Engine/SCS_Node.h"
@@ -564,9 +565,66 @@ void UBlueprint::PostDuplicate(bool bDuplicateForPIE)
 	}
 }
 
+extern COREUOBJECT_API bool GMinimalCompileOnLoad;
+
 UClass* UBlueprint::RegenerateClass(UClass* ClassToRegenerate, UObject* PreviousCDO, TArray<UObject*>& ObjLoaded)
 {
-	return FBlueprintEditorUtils::RegenerateBlueprintClass(this, ClassToRegenerate, PreviousCDO, ObjLoaded);
+	if(GMinimalCompileOnLoad)
+	{
+		// ensure that we have UProperties for any properties declared in the blueprint:
+		check(GeneratedClass && !bHasBeenRegenerated); // we probably need to create a class if none was loaded off disk
+		
+		UPackage* Package = Cast<UPackage>(GetOutermost());
+		bool bIsPackageDirty = Package ? Package->IsDirty() : false;
+
+		FBlueprintEditorUtils::RefreshVariables(this);
+
+		// clone generated class into skeleton class:
+		FObjectDuplicationParameters Params(GeneratedClass, GeneratedClass->GetOuter());
+		if (SimpleConstructionScript)
+		{
+			Params.DuplicationSeed.Add(SimpleConstructionScript, SimpleConstructionScript);
+		}
+		Params.ApplyFlags = RF_Transient;
+		Params.DestName = *(FString("SKEL_") + GeneratedClass->GetName());
+		SkeletonGeneratedClass = (UClass*)StaticDuplicateObjectEx(Params);
+		SkeletonGeneratedClass->ClassDefaultObject = nullptr;
+		if (UClass* Super = SkeletonGeneratedClass->GetSuperClass())
+		{
+			if (UBlueprint* GeneratedByBP = Cast<UBlueprint>(Super->ClassGeneratedBy))
+			{
+				check(GeneratedByBP->SkeletonGeneratedClass != nullptr);
+				SkeletonGeneratedClass->SetSuperStruct(GeneratedByBP->SkeletonGeneratedClass);
+			}
+		}
+		SkeletonGeneratedClass->Bind();
+		SkeletonGeneratedClass->StaticLink(true);
+		if(UBlueprintGeneratedClass* BPGC_Skel = Cast<UBlueprintGeneratedClass>(SkeletonGeneratedClass))
+		{
+			if( BPGC_Skel->UberGraphFunction )
+			{
+				BPGC_Skel->UberGraphFunction->StaticLink(true);
+			}
+		}
+		SkeletonGeneratedClass->GetDefaultObject(); // forces CDO to generate
+
+		FBlueprintEditorUtils::PreloadConstructionScript( this );
+
+		FBlueprintCompilationManager::NotifyBlueprintLoaded( this ); 
+		
+		if( Package )
+		{
+			Package->SetDirtyFlag(bIsPackageDirty);
+		}
+
+		// give the world the generated class... The only time this will be changed is if the generated class
+		// was null:
+		return GeneratedClass;
+	}
+	else
+	{
+		return FBlueprintEditorUtils::RegenerateBlueprintClass(this, ClassToRegenerate, PreviousCDO, ObjLoaded);
+	}
 }
 
 void UBlueprint::RemoveGeneratedClasses()
@@ -1092,62 +1150,81 @@ void UBlueprint::BeginCacheForCookedPlatformData(const ITargetPlatform *TargetPl
 {
 	Super::BeginCacheForCookedPlatformData(TargetPlatform);
 
-	// Only cook component data if the setting is enabled and this is an Actor-based Blueprint class.
-	if (GeneratedClass && GeneratedClass->IsChildOf<AActor>() && GetDefault<UCookerSettings>()->bCookBlueprintComponentTemplateData)
+	if (GeneratedClass && GeneratedClass->IsChildOf<AActor>())
 	{
-		int32 NumCookedComponents = 0;
-		const double StartTime = FPlatformTime::Seconds();
-
-		UBlueprintGeneratedClass* BPGClass = Cast<UBlueprintGeneratedClass>(*GeneratedClass);
-		if (BPGClass)
+		// Only cook component data if the setting is enabled and this is an Actor-based Blueprint class.
+		if (GetDefault<UCookerSettings>()->bCookBlueprintComponentTemplateData)
 		{
-			// Cook all overridden SCS component node templates inherited from the parent class hierarchy.
-			if (UInheritableComponentHandler* TargetInheritableComponentHandler = BPGClass->GetInheritableComponentHandler())
+			int32 NumCookedComponents = 0;
+			const double StartTime = FPlatformTime::Seconds();
+
+			UBlueprintGeneratedClass* BPGClass = Cast<UBlueprintGeneratedClass>(*GeneratedClass);
+			if (BPGClass)
 			{
-				for (auto RecordIt = TargetInheritableComponentHandler->CreateRecordIterator(); RecordIt; ++RecordIt)
+				// Cook all overridden SCS component node templates inherited from the parent class hierarchy.
+				if (UInheritableComponentHandler* TargetInheritableComponentHandler = BPGClass->GetInheritableComponentHandler())
 				{
-					if (!RecordIt->CookedComponentInstancingData.bIsValid)
+					for (auto RecordIt = TargetInheritableComponentHandler->CreateRecordIterator(); RecordIt; ++RecordIt)
 					{
-						// Note: This will currently block until finished.
-						// @TODO - Make this an async task so we can potentially cook instancing data for multiple components in parallel.
-						FBlueprintEditorUtils::BuildComponentInstancingData(RecordIt->ComponentTemplate, RecordIt->CookedComponentInstancingData);
-						++NumCookedComponents;
+						if (!RecordIt->CookedComponentInstancingData.bIsValid)
+						{
+							// Note: This will currently block until finished.
+							// @TODO - Make this an async task so we can potentially cook instancing data for multiple components in parallel.
+							FBlueprintEditorUtils::BuildComponentInstancingData(RecordIt->ComponentTemplate, RecordIt->CookedComponentInstancingData);
+							++NumCookedComponents;
+						}
 					}
 				}
-			}
 
-			// Cook all SCS component templates that are owned by this class.
-			if (BPGClass->SimpleConstructionScript)
-			{
-				for (auto Node : BPGClass->SimpleConstructionScript->GetAllNodes())
+				// Cook all SCS component templates that are owned by this class.
+				if (BPGClass->SimpleConstructionScript)
 				{
-					if (!Node->CookedComponentInstancingData.bIsValid)
+					for (auto Node : BPGClass->SimpleConstructionScript->GetAllNodes())
+					{
+						if (!Node->CookedComponentInstancingData.bIsValid)
+						{
+							// Note: This will currently block until finished.
+							// @TODO - Make this an async task so we can potentially cook instancing data for multiple components in parallel.
+							FBlueprintEditorUtils::BuildComponentInstancingData(Node->ComponentTemplate, Node->CookedComponentInstancingData);
+							++NumCookedComponents;
+						}
+					}
+				}
+
+				// Cook all UCS/AddComponent node templates that are owned by this class.
+				for (UActorComponent* ComponentTemplate : BPGClass->ComponentTemplates)
+				{
+					FBlueprintCookedComponentInstancingData& CookedComponentInstancingData = BPGClass->CookedComponentInstancingData.FindOrAdd(ComponentTemplate->GetFName());
+					if (!CookedComponentInstancingData.bIsValid)
 					{
 						// Note: This will currently block until finished.
 						// @TODO - Make this an async task so we can potentially cook instancing data for multiple components in parallel.
-						FBlueprintEditorUtils::BuildComponentInstancingData(Node->ComponentTemplate, Node->CookedComponentInstancingData);
+						FBlueprintEditorUtils::BuildComponentInstancingData(ComponentTemplate, CookedComponentInstancingData);
 						++NumCookedComponents;
 					}
 				}
 			}
 
-			// Cook all UCS/AddComponent node templates that are owned by this class.
-			for (UActorComponent* ComponentTemplate : BPGClass->ComponentTemplates)
+			if (NumCookedComponents > 0)
 			{
-				FBlueprintCookedComponentInstancingData& CookedComponentInstancingData = BPGClass->CookedComponentInstancingData.FindOrAdd(ComponentTemplate->GetFName());
-				if (!CookedComponentInstancingData.bIsValid)
-				{
-					// Note: This will currently block until finished.
-					// @TODO - Make this an async task so we can potentially cook instancing data for multiple components in parallel.
-					FBlueprintEditorUtils::BuildComponentInstancingData(ComponentTemplate, CookedComponentInstancingData);
-					++NumCookedComponents;
-				}
+				UE_LOG(LogBlueprint, Log, TEXT("%s: Cooked %d component(s) in %.02g ms"), *GetName(), NumCookedComponents, (FPlatformTime::Seconds() - StartTime) * 1000.0);
 			}
 		}
 
-		if (NumCookedComponents > 0)
+		// If we're using the "exclusive" nativization method and this Blueprint is not selected for nativization, determine if any of its parent Blueprints are selected.
+		if (GetDefault<UProjectPackagingSettings>()->BlueprintNativizationMethod == EProjectPackagingBlueprintNativizationMethod::Exclusive && !bSelectedForNativization)
 		{
-			UE_LOG(LogBlueprint, Log, TEXT("%s: Cooked %d component(s) in %.02g ms"), *GetName(), NumCookedComponents, (FPlatformTime::Seconds() - StartTime) * 1000.0);
+			TArray<UBlueprint*> ParentBlueprints;
+			GetBlueprintHierarchyFromClass(GeneratedClass, ParentBlueprints);
+			for (UBlueprint *ParentBP : ParentBlueprints)
+			{
+				if (ParentBP->bSelectedForNativization)
+				{
+					UBlueprintGeneratedClass* BPGC = CastChecked<UBlueprintGeneratedClass>(GeneratedClass);
+					BPGC->bHasNativizedParent = true;
+					break;
+				}
+			}
 		}
 	}
 }
diff --git a/Engine/Source/Runtime/Engine/Private/BlueprintGeneratedClass.cpp b/Engine/Source/Runtime/Engine/Private/BlueprintGeneratedClass.cpp
--- a/Engine/Source/Runtime/Engine/Private/BlueprintGeneratedClass.cpp
+++ b/Engine/Source/Runtime/Engine/Private/BlueprintGeneratedClass.cpp
@@ -31,6 +31,7 @@ UBlueprintGeneratedClass::UBlueprintGeneratedClass(const FObjectInitializer& Obj
 {
 	NumReplicatedProperties = 0;
 	bHasInstrumentation = false;
+	bHasNativizedParent = false;
 }
 
 void UBlueprintGeneratedClass::PostInitProperties()
@@ -942,6 +943,124 @@ void UBlueprintGeneratedClass::CreateComponentsForActor(const UClass* ThisClass,
 	}
 }
 
+void UBlueprintGeneratedClass::CheckAndApplyComponentTemplateOverrides(AActor* Actor)
+{
+	check(Actor != nullptr);
+
+	// Get the Blueprint class hierarchy (if valid).
+	TArray<const UBlueprintGeneratedClass*> ParentBPClassStack;
+	GetGeneratedClassesHierarchy(Actor->GetClass(), ParentBPClassStack);
+	if (ParentBPClassStack.Num() > 0)
+	{
+		// If the nearest native antecedent is also a nativized BP class, we may have an override
+		// in an ICH for some part of the non-native BP class hierarchy that also inherits from it.
+		if (UDynamicClass* ParentDynamicClass = Cast<UDynamicClass>(ParentBPClassStack[ParentBPClassStack.Num() - 1]->GetSuperClass()))
+		{
+			// Helper class for property checks during serialization.
+			class FComponentOverridePropertyArchive : public FArchive
+			{
+			public:
+				FComponentOverridePropertyArchive()
+					: FArchive()
+				{
+					ArIsSaving = true;
+
+					// Include properties that would normally skip tagged serialization (e.g. bulk serialization of array properties).
+					ArPortFlags |= PPF_ForceTaggedSerialization;
+				}
+
+				virtual bool ShouldSkipProperty(const UProperty* InProperty) const override
+				{
+					// These are properties that can be ignored at runtime in a cooked build (i.e. they indicate values that don't need to be checked or copied).
+					return (!InProperty
+						|| InProperty->IsEditorOnlyProperty()
+						|| InProperty->HasAnyPropertyFlags(CPF_Transient | CPF_ContainsInstancedReference | CPF_InstancedReference)
+						|| !InProperty->HasAnyPropertyFlags(CPF_Edit | CPF_Interp));
+				}
+			} ComponentOverridePropertyChecker;
+
+			// Get all default subobjects owned by the nativized antecedent's CDO.
+			// Note: This will also include all other inherited default subobjects.
+			TArray<UObject*> DefaultSubobjects;
+			ParentDynamicClass->GetDefaultObjectSubobjects(DefaultSubobjects);
+
+			// Pick out only the UActorComponent-based subobjects and cache them to use for checking below.
+			TArray<UActorComponent*> NativizedParentClassComponentSubobjects;
+			for (UObject* DefaultSubobject : DefaultSubobjects)
+			{
+				if (UActorComponent* ComponentSubobject = Cast<UActorComponent>(DefaultSubobject))
+				{
+					NativizedParentClassComponentSubobjects.Add(ComponentSubobject);
+				}
+			}
+
+			// Now check each non-native BP class (on up to the given Actor) for any inherited component template overrides, and manually apply default value overrides as we go.
+			for (int32 i = ParentBPClassStack.Num() - 1; i >= 0; i--)
+			{
+				const UBlueprintGeneratedClass* CurrentBPGClass = ParentBPClassStack[i];
+				check(CurrentBPGClass);
+
+				UInheritableComponentHandler* ICH = const_cast<UBlueprintGeneratedClass*>(CurrentBPGClass)->GetInheritableComponentHandler();
+				if (ICH && NativizedParentClassComponentSubobjects.Num() > 0)
+				{
+					// Check each default subobject that we've inherited from the antecedent class
+					for (UActorComponent* NativizedComponentSubobject : NativizedParentClassComponentSubobjects)
+					{
+						const FName NativizedComponentSubobjectName = NativizedComponentSubobject->GetFName();
+						FComponentKey ComponentKey = ICH->FindKey(NativizedComponentSubobjectName);
+						if (ComponentKey.IsValid() && ComponentKey.IsSCSKey())
+						{
+							const UClass* NativizedComponentSubobjectClass = NativizedComponentSubobject->GetClass();
+
+							// This is the override template object that's stored within the non-native Blueprint class
+							const UObject* NativizedComponentSubobjectTemplate = ICH->GetOverridenComponentTemplate(ComponentKey);
+
+							// This is the instance of the inherited component subobject that's owned by the given Actor instance
+							UObject* NativizedComponentSubobjectInstance = Actor->GetDefaultSubobjectByName(NativizedComponentSubobjectName);
+
+							// Preconditions (must be valid)
+							check(NativizedComponentSubobjectClass != nullptr);
+							check(NativizedComponentSubobjectTemplate != nullptr);
+							check(NativizedComponentSubobjectInstance != nullptr);
+							check(NativizedComponentSubobjectClass == NativizedComponentSubobjectTemplate->GetClass());
+							check(NativizedComponentSubobjectClass == NativizedComponentSubobjectInstance->GetClass());
+
+							// This is the native DSO (archetype) on which the component subobject instance is currently based
+							const UObject* NativizedComponentSubobjectArchetype = NativizedComponentSubobjectInstance->GetArchetype();
+							check(NativizedComponentSubobjectArchetype != nullptr);
+
+							// Apply override template defaults to the instance that's owned by the given Actor (when appropriate).
+							for (TFieldIterator<UProperty> PropertyIt(NativizedComponentSubobjectClass); PropertyIt; ++PropertyIt)
+							{
+								UProperty* Property = *PropertyIt;
+								if (Property->ShouldSerializeValue(ComponentOverridePropertyChecker))
+								{
+									for (int32 ArrayIdx = 0; ArrayIdx < Property->ArrayDim; ++ArrayIdx)
+									{
+										uint8* DstValuePtr = Property->ContainerPtrToValuePtr<uint8>(NativizedComponentSubobjectInstance, ArrayIdx);
+										const uint8* SrcValuePtr = Property->ContainerPtrToValuePtr<uint8>(NativizedComponentSubobjectTemplate, ArrayIdx);
+										const uint8* DefaultValuePtr = Property->ContainerPtrToValuePtr<uint8>(NativizedComponentSubobjectArchetype, ArrayIdx);
+
+										// Only copy the value (1) when it matches the archetype and (2) does not match the override template. Failing the
+										// first check indicates that it was already overridden at the instance level (and we don't want to overwrite that).
+										if (Property->Identical(DstValuePtr, DefaultValuePtr) && !Property->Identical(DstValuePtr, SrcValuePtr))
+										{
+											Property->CopySingleValue(DstValuePtr, SrcValuePtr);
+										}
+									}
+								}
+							}
+
+							// There can only be a single match, so we can stop searching now.
+							break;
+						}
+					}
+				}
+			}
+		}
+	}
+}
+
 uint8* UBlueprintGeneratedClass::GetPersistentUberGraphFrame(UObject* Obj, UFunction* FuncToCheck) const
 {
 	if (Obj && UsePersistentUberGraphFrame() && UberGraphFramePointerProperty && UberGraphFunction)
diff --git a/Engine/Source/Runtime/Engine/Private/InheritableComponentHandler.cpp b/Engine/Source/Runtime/Engine/Private/InheritableComponentHandler.cpp
--- a/Engine/Source/Runtime/Engine/Private/InheritableComponentHandler.cpp
+++ b/Engine/Source/Runtime/Engine/Private/InheritableComponentHandler.cpp
@@ -278,7 +278,7 @@ bool UInheritableComponentHandler::IsRecordValid(const FComponentOverrideRecord&
 		return false;
 	}
 
-	UBlueprintGeneratedClass* ComponentOwner = Record.ComponentKey.GetComponentOwner();
+	UClass* ComponentOwner = Record.ComponentKey.GetComponentOwner();
 	if (!ComponentOwner || !OwnerClass->IsChildOf(ComponentOwner))
 	{
 		return false;
@@ -515,7 +515,7 @@ FComponentKey::FComponentKey(const USCS_Node* SCSNode) : OwnerClass(nullptr)
 	if (SCSNode)
 	{
 		const USimpleConstructionScript* ParentSCS = SCSNode->GetSCS();
-		OwnerClass      = ParentSCS ? Cast<UBlueprintGeneratedClass>(ParentSCS->GetOwnerClass()) : nullptr;
+		OwnerClass      = ParentSCS ? ParentSCS->GetOwnerClass() : nullptr;
 		AssociatedGuid  = SCSNode->VariableGuid;
 		SCSVariableName = SCSNode->GetVariableName();
 	}
@@ -524,7 +524,7 @@ FComponentKey::FComponentKey(const USCS_Node* SCSNode) : OwnerClass(nullptr)
 #if WITH_EDITOR
 FComponentKey::FComponentKey(UBlueprint* Blueprint, const FUCSComponentId& UCSComponentID)
 {
-	OwnerClass     = CastChecked<UBlueprintGeneratedClass>(Blueprint->GeneratedClass);
+	OwnerClass     = Blueprint->GeneratedClass;
 	AssociatedGuid = UCSComponentID.GetAssociatedGuid();
 }
 #endif // WITH_EDITOR
@@ -536,7 +536,7 @@ bool FComponentKey::Match(const FComponentKey OtherKey) const
 
 USCS_Node* FComponentKey::FindSCSNode() const
 {
-	USimpleConstructionScript* ParentSCS = (OwnerClass && IsSCSKey()) ? OwnerClass->SimpleConstructionScript : nullptr;
+	USimpleConstructionScript* ParentSCS = (OwnerClass && IsSCSKey()) ? CastChecked<UBlueprintGeneratedClass>(OwnerClass)->SimpleConstructionScript : nullptr;
 	return ParentSCS ? ParentSCS->FindSCSNodeByGuid(AssociatedGuid) : nullptr;
 }
 
diff --git a/Engine/Source/Runtime/Engine/Private/KismetMathLibrary.cpp b/Engine/Source/Runtime/Engine/Private/KismetMathLibrary.cpp
--- a/Engine/Source/Runtime/Engine/Private/KismetMathLibrary.cpp
+++ b/Engine/Source/Runtime/Engine/Private/KismetMathLibrary.cpp
@@ -198,6 +198,20 @@ float UKismetMathLibrary::Hypotenuse(float Width, float Height)
 	return Max * FMath::Sqrt(1.f + Ratio * Ratio);
 }
 
+float UKismetMathLibrary::Log(float A, float Base)
+{
+	float Result = 0.f;
+	if(Base <= 0.f)
+	{
+		FFrame::KismetExecutionMessage(TEXT("Divide by zero: Log"), ELogVerbosity::Warning, DivideByZeroWarning);
+	}
+	else
+	{
+		Result = FMath::Loge(A) / FMath::Loge(Base);
+	}
+	return Result;
+}
+
 int32 UKismetMathLibrary::FMod(float Dividend, float Divisor, float& Remainder)
 {
 	int32 Result;
diff --git a/Engine/Source/Runtime/Engine/Private/Level.cpp b/Engine/Source/Runtime/Engine/Private/Level.cpp
--- a/Engine/Source/Runtime/Engine/Private/Level.cpp
+++ b/Engine/Source/Runtime/Engine/Private/Level.cpp
@@ -634,6 +634,13 @@ void ULevel::PostLoad()
 		LevelSimplification[Index].PostLoadDeprecated();
 	}
 
+	if (LevelScriptActor)
+	{
+		if (ULevelScriptBlueprint* LevelBlueprint = Cast<ULevelScriptBlueprint>(LevelScriptActor->GetClass()->ClassGeneratedBy))
+		{
+			FBlueprintEditorUtils::FixLevelScriptActorBindings(LevelScriptActor, LevelBlueprint);
+		}
+	}
 #endif
 }
 
diff --git a/Engine/Source/Runtime/Engine/Private/UserDefinedEnum.cpp b/Engine/Source/Runtime/Engine/Private/UserDefinedEnum.cpp
--- a/Engine/Source/Runtime/Engine/Private/UserDefinedEnum.cpp
+++ b/Engine/Source/Runtime/Engine/Private/UserDefinedEnum.cpp
@@ -97,6 +97,17 @@ FString UUserDefinedEnum::GenerateNewEnumeratorName()
 	return EnumNameString;
 }
 
+FText UUserDefinedEnum::GetDisplayNameText(int32 NameIndex) const
+{
+	FText DisplayName = FEnumEditorUtils::GetEnumeratorDisplayName(this, NameIndex);
+	if (DisplayName.IsEmpty())
+	{
+		DisplayName = Super::GetDisplayNameText(NameIndex);
+	}
+
+	return DisplayName;
+}
+
 #endif	// WITH_EDITOR
 
 int64 UUserDefinedEnum::ResolveEnumerator(FArchive& Ar, int64 EnumeratorValue) const
