Backing out 3350092 - in addition to activating a nascent bad warning, this change introduces errors in blueprints that have an input ref to a pure node that is connected to a return node.

diff --git a/Engine/Source/Editor/KismetCompiler/Private/KismetCompiler.cpp b/Engine/Source/Editor/KismetCompiler/Private/KismetCompiler.cpp
--- a/Engine/Source/Editor/KismetCompiler/Private/KismetCompiler.cpp
+++ b/Engine/Source/Editor/KismetCompiler/Private/KismetCompiler.cpp
@@ -86,11 +86,7 @@ namespace
 	{
 		for (UEdGraphNode* Node : Graph->Nodes)
 		{
-			// UK2Node_FunctionResult can't be pruned because we need it to determine function outputs. 
-			// If we prune UK2Node_FunctionResult we generate a function  that is missing user defined
-			// outputs. UK2Node_FunctionResult has no exec output so it will not bring in any extra
-			// nodes, but by keeping it we ensure that the expected output parameters are generated:
-			const bool bRootSetByType = Node && (Node->IsA<UK2Node_FunctionEntry>() || Node->IsA<UK2Node_Event>() || Node->IsA<UK2Node_Timeline>() || Node->IsA<UK2Node_FunctionResult>());
+			const bool bRootSetByType = Node && (Node->IsA<UK2Node_FunctionEntry>() || Node->IsA<UK2Node_Event>() || Node->IsA<UK2Node_Timeline>());
 			UK2Node* K2Node = Cast<UK2Node>(Node);
 			bool bIsRootSet = bRootSetByType || (K2Node && K2Node->IsNodeRootSet());
 
