Copying //UE4/Dev-Framework to //UE4/Dev-Main (Source: //UE4/Dev-Framework @ 3431384)

diff --git a/Engine/Source/Editor/Kismet/Private/Tests/BlueprintEditorTests.cpp b/Engine/Source/Editor/Kismet/Private/Tests/BlueprintEditorTests.cpp
--- a/Engine/Source/Editor/Kismet/Private/Tests/BlueprintEditorTests.cpp
+++ b/Engine/Source/Editor/Kismet/Private/Tests/BlueprintEditorTests.cpp
@@ -607,7 +607,7 @@ namespace BlueprintEditorPromotionUtils
 	static void AddStringMemberValue(UBlueprint* InBlueprint, const FName& VariableName)
 	{
 		const UEdGraphSchema_K2* K2Schema = GetDefault<UEdGraphSchema_K2>();
-		FEdGraphPinType StringPinType(K2Schema->PC_String, TEXT(""), NULL, false, false, false, false, FEdGraphTerminalType());
+		FEdGraphPinType StringPinType(K2Schema->PC_String, FString(), nullptr, EPinContainerType::None, false, FEdGraphTerminalType());
 		FBlueprintEditorUtils::AddMemberVariable(InBlueprint, VariableName, StringPinType);
 	}
 
diff --git a/Engine/Source/Editor/Kismet/Private/UserDefinedStructureEditor.cpp b/Engine/Source/Editor/Kismet/Private/UserDefinedStructureEditor.cpp
--- a/Engine/Source/Editor/Kismet/Private/UserDefinedStructureEditor.cpp
+++ b/Engine/Source/Editor/Kismet/Private/UserDefinedStructureEditor.cpp
@@ -435,7 +435,7 @@ class FUserDefinedStructureLayout : public IDetailCustomNodeBuilder, public TSha
 public:
 	FUserDefinedStructureLayout(TWeakPtr<class FUserDefinedStructureDetails> InStructureDetails)
 		: StructureDetails(InStructureDetails)
-		, InitialPinType(GetDefault<UEdGraphSchema_K2>()->PC_Boolean, TEXT(""), NULL, false, false, false, false, FEdGraphTerminalType())
+		, InitialPinType(GetDefault<UEdGraphSchema_K2>()->PC_Boolean, FString(), nullptr, EPinContainerType::None, false, FEdGraphTerminalType())
 	{}
 
 	void OnChanged()
diff --git a/Engine/Source/Editor/Kismet/Public/BlueprintCompilationManager.h b/Engine/Source/Editor/Kismet/Public/BlueprintCompilationManager.h
--- a/Engine/Source/Editor/Kismet/Public/BlueprintCompilationManager.h
+++ b/Engine/Source/Editor/Kismet/Public/BlueprintCompilationManager.h
@@ -4,7 +4,30 @@
 
 #include "CoreMinimal.h" // for DLLEXPORT (KISMET_API)
 
+#include "Kismet2/KismetEditorUtilities.h"
+
 class UBlueprint;
+class FCompilerResultsLog;
+
+struct FBPCompileRequest
+{
+	explicit FBPCompileRequest(UBlueprint* InBPToCompile, EBlueprintCompileOptions InCompileOptions, FCompilerResultsLog* InClientResultsLog )
+		: BPToCompile(InBPToCompile)
+		, CompileOptions(InCompileOptions)
+		, ClientResultsLog(InClientResultsLog)
+	{
+	}
+
+	// BP that needs to be compiled:
+	UBlueprint* BPToCompile;
+
+	// Legacy options for blueprint compilation:
+	EBlueprintCompileOptions CompileOptions;
+	
+	// Clients can give us a results log if they want to parse or display it themselves, otherwise
+	// we will use a transient one:
+	FCompilerResultsLog* ClientResultsLog;
+};
 
 struct KISMET_API FBlueprintCompilationManager
 {
@@ -21,9 +44,11 @@ struct KISMET_API FBlueprintCompilationManager
 	static void FlushCompilationQueue(TArray<UObject*>* ObjLoaded = nullptr);
 
 	/**
-	 * Adds the blueprint to the compilation queue
+	 * Immediately compiles the blueprint, no expectation that related blueprints be subsequently compiled.
+	 * It will be significantly more efficient to queue blueprints and then flush the compilation queue
+	 * if there are several blueprints that require compilation (e.g. typical case on PIE):
 	 */
-	static void QueueForCompilation(UBlueprint* BPToCompile);
+	static void CompileSynchronously(const FBPCompileRequest& Request);
 	
 	/**
 	 * Adds a newly loaded blueprint to the compilation queue
diff --git a/Engine/Source/Editor/Kismet/Public/BlueprintEditor.h b/Engine/Source/Editor/Kismet/Public/BlueprintEditor.h
--- a/Engine/Source/Editor/Kismet/Public/BlueprintEditor.h
+++ b/Engine/Source/Editor/Kismet/Public/BlueprintEditor.h
@@ -714,6 +714,9 @@ protected:
 	void OnRestoreAllStructVarPins();
 	bool CanRestoreAllStructVarPins() const;
 
+	void OnResetPinToDefaultValue();
+	bool CanResetPinToDefaultValue() const;
+
 	void OnAddOptionPin();
 	bool CanAddOptionPin() const;
 
diff --git a/Engine/Source/Editor/Kismet/Public/SSCSEditor.h b/Engine/Source/Editor/Kismet/Public/SSCSEditor.h
--- a/Engine/Source/Editor/Kismet/Public/SSCSEditor.h
+++ b/Engine/Source/Editor/Kismet/Public/SSCSEditor.h
@@ -121,7 +121,7 @@ public:
 	/**
 	 * @return Whether or not this node is a child (direct or indirect) of the given node.
 	 */
-	bool IsAttachedTo(FSCSEditorTreeNodePtrType InNodePtr) const;
+	bool IsAttachedTo(FSCSEditorTreeNodePtrType InNodePtr) const;	
 
 	/**
 	 * Finds the closest ancestor node in the given node set.
@@ -288,8 +288,22 @@ public:
 	/** Sets up the delegate for renaming a component */
 	void SetRenameRequestedDelegate(FOnRenameRequested InRenameRequested) { RenameRequestedDelegate = InRenameRequested; }
 
+	/** Query that determines if this item should be filtered out or not */
+	bool IsFlaggedForFiltration() const 
+	{
+		return ensureMsgf(FilterFlags != EFilteredState::Unknown, TEXT("Querying a bad filtration state.")) ? 
+			(FilterFlags & EFilteredState::FilteredInMask) == 0 : false; 
+	}
+
+	/** Refreshes this item's filtration state. Use bUpdateParent to make sure the parent's EFilteredState::ChildMatches flag is properly updated based off the new state */
+	void UpdateCachedFilterState(bool bMatchesFilter, bool bUpdateParent);
 
 protected:
+	/** Updates the EFilteredState::ChildMatches flag, based off of children's current state */
+	void RefreshCachedChildFilterState(bool bUpdateParent);
+	/** Used to update the EFilteredState::ChildMatches flag for parent nodes, when this item's filtration state has changed */
+	void ApplyFilteredStateToParent();
+	
 	bool GetAndClearNonTransactionalRenameFlag()
 	{
 		const bool bResult = bNonTransactionalRename;
@@ -311,6 +325,17 @@ private:
 	/** Handles rename requests */
 	bool bNonTransactionalRename;
 	FOnRenameRequested RenameRequestedDelegate;
+
+	enum EFilteredState
+	{
+		FilteredOut    = 0x00,
+		MatchesFilter  = (1 << 0),
+		ChildMatches   = (1 << 1),
+
+		FilteredInMask = (MatchesFilter | ChildMatches),
+		Unknown = 0xFC // ~FilteredInMask
+	};
+	uint8 FilterFlags;
 };
 
 //////////////////////////////////////////////////////////////////////////
@@ -784,7 +809,7 @@ public:
 	   @param Asset       					(In) Optional asset to assign to the component
 	   @param bSkipMarkBlueprintModified 	(In) Optionally skip marking this blueprint as modified (e.g. if we're handling that externally)
 	   @return The reference of the newly created ActorComponent */
-	UActorComponent* AddNewComponent(UClass* NewComponentClass, UObject* Asset, const bool bSkipMarkBlueprintModified = false );
+	UActorComponent* AddNewComponent(UClass* NewComponentClass, UObject* Asset, const bool bSkipMarkBlueprintModified = false, bool bSetFocusToNewItem = true);
 
 	/** Adds a new SCS Node to the component Table
 	   @param NewNode	(In) The SCS node to add
@@ -944,6 +969,9 @@ public:
 	/** Handler for recursively expanding/collapsing items */
 	void SetItemExpansionRecursive(FSCSEditorTreeNodePtrType Model, bool bInExpansionState);
 
+	/** Callback for the action trees to get the filter text */
+	FText GetFilterText() const;
+
 protected:
 	FString GetSelectedClassText() const;
 
@@ -1059,7 +1087,6 @@ protected:
 	 * @return The new class that was created
 	 */
 	UClass* CreateNewCPPComponent(TSubclassOf<UActorComponent> ComponentClass);
-
 	
 	/**
 	 * Creates a new Blueprint component from the specified class type
@@ -1069,6 +1096,22 @@ protected:
 	 */
 	UClass* CreateNewBPComponent(TSubclassOf<UActorComponent> ComponentClass);
 
+	/** Recursively updates the filtered state for each component item */
+	void OnFilterTextChanged(const FText& InFilterText);
+
+	/** 
+	 * Compares the filter bar's text with the item's component name. Use 
+	 * bRecursive to refresh the state of child nodes as well. Returns true if 
+	 * the node is set to be filtered out 
+	 */
+	bool RefreshFilteredState(FSCSEditorTreeNodePtrType TreeNode, bool bRecursive);
+
+	/** 
+	 * Iterates the RootNodes list, and uses the cached filtered state to 
+	 * determine what items should be listed in the tree view. 
+	 */
+	void RebuildFilteredRootList();
+
 public:
 	/** Tree widget */
 	TSharedPtr<SSCSTreeType> SCSTreeWidget;
@@ -1115,6 +1158,9 @@ private:
 	/** Root set of components (contains the root scene component and any non-scene component nodes) */
 	TArray<FSCSEditorTreeNodePtrType> RootComponentNodes;
 
+	/** The list of nodes used for the UI (a filtered version of RootNodes) */
+	TArray<FSCSEditorTreeNodePtrType> FilteredRootNodes;
+
 	/* Root Tree Node (for scene components) */
 	TSharedPtr<FSCSEditorTreeNode> RootTreeNode;
 
@@ -1135,4 +1181,7 @@ private:
 
 	/** TRUE if this SCSEditor is currently the target of a diff */
 	bool bIsDiffing;
+
+	/** The filter box that handles filtering for the tree. */
+	TSharedPtr< SSearchBox > FilterBox;
 };
diff --git a/Engine/Source/Editor/KismetCompiler/Private/AnimBlueprintCompiler.cpp b/Engine/Source/Editor/KismetCompiler/Private/AnimBlueprintCompiler.cpp
--- a/Engine/Source/Editor/KismetCompiler/Private/AnimBlueprintCompiler.cpp
+++ b/Engine/Source/Editor/KismetCompiler/Private/AnimBlueprintCompiler.cpp
@@ -82,7 +82,7 @@ bool FAnimBlueprintCompiler::FEffectiveConstantRecord::Apply(UObject* Object)
 
 		if (ArrayHelper.IsValidIndex(ArrayIndex))
 		{
-			FBlueprintEditorUtils::ImportKismetDefaultValueToProperty(LiteralSourcePin, ArrayProperty->Inner, ArrayHelper.GetRawPtr(ArrayIndex), Object);
+			FBlueprintEditorUtils::PropertyValueFromString_Direct(ArrayProperty->Inner, LiteralSourcePin->GetDefaultAsString(), ArrayHelper.GetRawPtr(ArrayIndex));
 		}
 		else
 		{
@@ -91,7 +91,7 @@ bool FAnimBlueprintCompiler::FEffectiveConstantRecord::Apply(UObject* Object)
 	}
 	else
 	{
-		FBlueprintEditorUtils::ImportKismetDefaultValueToProperty(LiteralSourcePin, ConstantProperty, PropertyPtr, Object);
+		FBlueprintEditorUtils::PropertyValueFromString_Direct(ConstantProperty, LiteralSourcePin->GetDefaultAsString(), PropertyPtr);
 	}
 
 	return true;
@@ -225,7 +225,7 @@ void FAnimBlueprintCompiler::CreateEvaluationHandlerStruct(UAnimGraphNode_Base*
 		// Does it get serviced by this handler?
 		if (FAnimNodeSinglePropertyHandler* SourceInfo = Record.ServicedProperties.Find(PropertyName))
 		{
-			if (TargetPin->PinType.bIsArray)
+			if (TargetPin->PinType.IsArray())
 			{
 				// Grab the array that we need to set members for
 				UK2Node_StructMemberGet* FetchArrayNode = SpawnIntermediateNode<UK2Node_StructMemberGet>(VisualAnimNode, ConsolidatedEventGraph);
@@ -271,6 +271,7 @@ void FAnimBlueprintCompiler::CreateEvaluationHandlerStruct(UAnimGraphNode_Base*
 			}
 			else
 			{
+				check(!TargetPin->PinType.IsContainer())
 				// Single property
 				if (SourceInfo->CopyRecords.Num() > 0 && SourceInfo->CopyRecords[0].DestPin != nullptr)
 				{
@@ -351,7 +352,7 @@ void FAnimBlueprintCompiler::CreateEvaluationHandlerInstance(UAnimGraphNode_Base
 			// Find the property pin on the set node and configure
 			for(UEdGraphPin* TargetPin : VarAssignNode->Pins)
 			{
-				if(TargetPin->PinType.bIsArray)
+				if(TargetPin->PinType.IsContainer())
 				{
 					// Currently unsupported
 					continue;
diff --git a/Engine/Source/Editor/KismetCompiler/Private/KismetCompiler.cpp b/Engine/Source/Editor/KismetCompiler/Private/KismetCompiler.cpp
--- a/Engine/Source/Editor/KismetCompiler/Private/KismetCompiler.cpp
+++ b/Engine/Source/Editor/KismetCompiler/Private/KismetCompiler.cpp
@@ -287,7 +287,8 @@ void FKismetCompilerContext::CleanAndSanitizeClass(UBlueprintGeneratedClass* Cla
 
 	// Set properties we need to regenerate the class with
 	ClassToClean->PropertyLink = ParentClass->PropertyLink;
-	ClassToClean->ClassWithin = ParentClass;
+	ClassToClean->SetSuperStruct(ParentClass);
+	ClassToClean->ClassWithin = ParentClass->ClassWithin ? ParentClass->ClassWithin : UObject::StaticClass();
 	ClassToClean->ClassConfigName = ClassToClean->IsNative() ? FName(ClassToClean->StaticConfigName()) : ParentClass->ClassConfigName;
 	ClassToClean->DebugData = FBlueprintDebugData();
 }
@@ -547,6 +548,7 @@ void FKismetCompilerContext::ValidateVariableNames()
 			if (OldVarName != NewVarName)
 			{
 				MessageLog.Warning(*FString::Printf(*LOCTEXT("MemberVariableConflictWarning", "Found a member variable with a conflicting name (%s) - changed to %s.").ToString(), *VarNameStr, *NewVarName.ToString()));
+				TGuardValue<bool> LockDependencies(Blueprint->bCachedDependenciesUpToDate, Blueprint->bCachedDependenciesUpToDate);
 				FBlueprintEditorUtils::RenameMemberVariable(Blueprint, OldVarName, NewVarName);
 			}
 		}
@@ -672,7 +674,7 @@ void FKismetCompilerContext::CreateClassVariablesFromBlueprint()
 			continue;
 		}
 
-		FEdGraphPinType TimelinePinType(Schema->PC_Object, TEXT(""), UTimelineComponent::StaticClass(), false, false, false, false, FEdGraphTerminalType());
+		FEdGraphPinType TimelinePinType(Schema->PC_Object, FString(), UTimelineComponent::StaticClass(), EPinContainerType::None, false, FEdGraphTerminalType());
 
 		// Previously UTimelineComponent object has exactly the same name as UTimelineTemplate object (that obj was in blueprint)
 		const FString TimelineVariableName = UTimelineTemplate::TimelineTemplateNameToVariableName(Timeline->GetFName());
@@ -685,22 +687,22 @@ void FKismetCompilerContext::CreateClassVariablesFromBlueprint()
 			TimelineToMemberVariableMap.Add(Timeline, TimelineProperty);
 		}
 
-		FEdGraphPinType DirectionPinType(Schema->PC_Byte, TEXT(""), FTimeline::GetTimelineDirectionEnum(), false, false, false, false, FEdGraphTerminalType());
+		FEdGraphPinType DirectionPinType(Schema->PC_Byte, FString(), FTimeline::GetTimelineDirectionEnum(), EPinContainerType::None, false, FEdGraphTerminalType());
 		CreateVariable(Timeline->GetDirectionPropertyName(), DirectionPinType);
 
-		FEdGraphPinType FloatPinType(Schema->PC_Float, TEXT(""), NULL, false, false, false, false, FEdGraphTerminalType());
+		FEdGraphPinType FloatPinType(Schema->PC_Float, FString(), nullptr, EPinContainerType::None, false, FEdGraphTerminalType());
 		for(int32 i=0; i<Timeline->FloatTracks.Num(); i++)
 		{
 			CreateVariable(Timeline->GetTrackPropertyName(Timeline->FloatTracks[i].TrackName), FloatPinType);
 		}
 
-		FEdGraphPinType VectorPinType(Schema->PC_Struct, TEXT(""), VectorStruct, false, false, false, false, FEdGraphTerminalType());
+		FEdGraphPinType VectorPinType(Schema->PC_Struct, FString(), VectorStruct, EPinContainerType::None, false, FEdGraphTerminalType());
 		for(int32 i=0; i<Timeline->VectorTracks.Num(); i++)
 		{
 			CreateVariable(Timeline->GetTrackPropertyName(Timeline->VectorTracks[i].TrackName), VectorPinType);
 		}
 
-		FEdGraphPinType LinearColorPinType(Schema->PC_Struct, TEXT(""), LinearColorStruct, false, false, false, false, FEdGraphTerminalType());
+		FEdGraphPinType LinearColorPinType(Schema->PC_Struct, FString(), LinearColorStruct, EPinContainerType::None, false, FEdGraphTerminalType());
 		for(int32 i=0; i<Timeline->LinearColorTracks.Num(); i++)
 		{
 			CreateVariable(Timeline->GetTrackPropertyName(Timeline->LinearColorTracks[i].TrackName), LinearColorPinType);
@@ -708,7 +710,7 @@ void FKismetCompilerContext::CreateClassVariablesFromBlueprint()
 	}
 
 	// Create a class property for any simple-construction-script created components that should be exposed
-	if (Blueprint->SimpleConstructionScript != NULL)
+	if (Blueprint->SimpleConstructionScript)
 	{
 		// Ensure that nodes have valid templates (This will remove nodes that have had the classes the inherited from removed
 		Blueprint->SimpleConstructionScript->ValidateNodeTemplates(MessageLog);
@@ -723,9 +725,8 @@ void FKismetCompilerContext::CreateClassVariablesFromBlueprint()
 				FName VarName = Node->GetVariableName();
 				if ((VarName != NAME_None) && (Node->ComponentClass != nullptr))
 				{
-					FEdGraphPinType Type(Schema->PC_Object, TEXT(""), Node->ComponentClass, false, false, false, false, FEdGraphTerminalType());
-					UProperty* NewProperty = CreateVariable(VarName, Type);
-					if (NewProperty != NULL)
+					FEdGraphPinType Type(Schema->PC_Object, FString(), Node->ComponentClass, EPinContainerType::None, false, FEdGraphTerminalType());
+					if (UProperty* NewProperty = CreateVariable(VarName, Type))
 					{
 						const FText CategoryName = Node->CategoryName.IsEmpty() ? FText::FromString(Blueprint->GetName()) : Node->CategoryName ;
 					
@@ -1021,24 +1022,10 @@ void FKismetCompilerContext::CreateUserDefinedLocalVariablesForFunction(FKismetF
 	for (int32 i = 0; i < Context.EntryPoint->LocalVariables.Num(); ++i)
 	{
 		FBPVariableDescription& Variable = Context.EntryPoint->LocalVariables[Context.EntryPoint->LocalVariables.Num() - (i + 1)];
+		UProperty* NewProperty = CreateUserDefinedLocalVariableForFunction(Variable, Context.Function, NewClass, FunctionPropertyStorageLocation, Schema, MessageLog);
 
-		// Create the property based on the variable description, scoped to the function
-		UProperty* NewProperty = FKismetCompilerUtilities::CreatePropertyOnScope(Context.Function, Variable.VarName, Variable.VarType, NewClass, 0, Schema, MessageLog);
 		if (NewProperty != NULL)
 		{
-			// Link this object to the tail of the list (so properties remain in the desired order)
-			*FunctionPropertyStorageLocation = NewProperty;
-			FunctionPropertyStorageLocation = &(NewProperty->Next);
-		}
-
-		if (NewProperty != NULL)
-		{
-			NewProperty->SetPropertyFlags(Variable.PropertyFlags);
-			NewProperty->SetMetaData(TEXT("FriendlyName"), *Variable.FriendlyName);
-			NewProperty->SetMetaData(TEXT("Category"), *Variable.Category.ToString());
-			NewProperty->RepNotifyFunc = Variable.RepNotifyFunc;
-			NewProperty->SetPropertyFlags(Variable.PropertyFlags);
-
 			if(!Variable.DefaultValue.IsEmpty())
 			{
 				SetPropertyDefaultValue(NewProperty, Variable.DefaultValue);
@@ -1047,6 +1034,26 @@ void FKismetCompilerContext::CreateUserDefinedLocalVariablesForFunction(FKismetF
 	}
 }
 
+UProperty* FKismetCompilerContext::CreateUserDefinedLocalVariableForFunction(const FBPVariableDescription& Variable, UFunction* Function, UBlueprintGeneratedClass* OwningClass, UField**& FunctionPropertyStorageLocation, const UEdGraphSchema_K2* Schema, FCompilerResultsLog& MessageLog)
+{
+	UProperty* NewProperty = FKismetCompilerUtilities::CreatePropertyOnScope(Function, Variable.VarName, Variable.VarType, OwningClass, 0, Schema, MessageLog);
+	
+	if(NewProperty)
+	{
+		// Link this object to the tail of the list (so properties remain in the desired order)
+		*FunctionPropertyStorageLocation = NewProperty;
+		FunctionPropertyStorageLocation = &(NewProperty->Next);
+		
+		NewProperty->SetPropertyFlags(Variable.PropertyFlags);
+		NewProperty->SetMetaData(TEXT("FriendlyName"), *Variable.FriendlyName);
+		NewProperty->SetMetaData(TEXT("Category"), *Variable.Category.ToString());
+		NewProperty->RepNotifyFunc = Variable.RepNotifyFunc;
+		NewProperty->SetPropertyFlags(Variable.PropertyFlags);
+	}
+	
+	return NewProperty;
+}
+
 void FKismetCompilerContext::SetPropertyDefaultValue(const UProperty* PropertyToSet, FString& Value)
 {
 	DefaultPropertyValueMap.Add(PropertyToSet->GetFName(), Value);
@@ -1981,7 +1988,16 @@ void FKismetCompilerContext::PostcompileFunction(FKismetFunctionContext& Context
  */
 void FKismetCompilerContext::FinishCompilingFunction(FKismetFunctionContext& Context)
 {
-	UFunction* Function = Context.Function;
+	SetCalculatedMetaDataAndFlags( Context.Function, CastChecked<UK2Node_FunctionEntry>(Context.EntryPoint), Schema );
+}
+
+void FKismetCompilerContext::SetCalculatedMetaDataAndFlags(UFunction* Function, UK2Node_FunctionEntry* EntryNode, const UEdGraphSchema_K2* K2Schema)
+{
+	if(!ensure(Function) || !ensure(EntryNode))
+	{
+		return;
+	}
+
 	Function->Bind();
 	Function->StaticLink(true);
 
@@ -2012,15 +2028,32 @@ void FKismetCompilerContext::FinishCompilingFunction(FKismetFunctionContext& Con
 		{
 			if (!Property->HasAnyPropertyFlags(CPF_ZeroConstructor))
 			{
-			Function->FirstPropertyToInit = Property;
-			Function->FunctionFlags |= FUNC_HasDefaults;
+				Function->FirstPropertyToInit = Property;
+				Function->FunctionFlags |= FUNC_HasDefaults;
 				break;
-		}
+			}
 		}
 	}
 
+	FKismetUserDeclaredFunctionMetadata& FunctionMetaData = EntryNode->MetaData;
+	if (!FunctionMetaData.Category.IsEmpty())
+	{
+		Function->SetMetaData(FBlueprintMetadata::MD_FunctionCategory, *FunctionMetaData.Category.ToString());
+	}
+
+	// Set up the function keywords
+	if (!FunctionMetaData.Keywords.IsEmpty())
+	{
+		Function->SetMetaData(FBlueprintMetadata::MD_FunctionKeywords, *FunctionMetaData.Keywords.ToString());
+	}
+
+	// Set up the function compact node title
+	if (!FunctionMetaData.CompactNodeTitle.IsEmpty())
+	{
+		Function->SetMetaData(FBlueprintMetadata::MD_CompactNodeTitle, *FunctionMetaData.CompactNodeTitle.ToString());
+	}
+
 	// Add in any extra user-defined metadata, like tooltip
-	UK2Node_FunctionEntry* EntryNode = CastChecked<UK2Node_FunctionEntry>(Context.EntryPoint);
 	if (!EntryNode->MetaData.ToolTip.IsEmpty())
 	{
 		Function->SetMetaData(FBlueprintMetadata::MD_Tooltip, *EntryNode->MetaData.ToolTip.ToString());
@@ -2038,15 +2071,21 @@ void FKismetCompilerContext::FinishCompilingFunction(FKismetFunctionContext& Con
 	{
 		UEdGraphPin* EntryPin = EntryNode->Pins[EntryPinIndex];
 		// No defaults for object/class pins
-		if(	!Schema->IsMetaPin(*EntryPin) && 
-			(EntryPin->PinType.PinCategory != Schema->PC_Object) && 
-			(EntryPin->PinType.PinCategory != Schema->PC_Class) && 
-			(EntryPin->PinType.PinCategory != Schema->PC_Interface) && 
+		if(	!K2Schema->IsMetaPin(*EntryPin) && 
+			(EntryPin->PinType.PinCategory != UEdGraphSchema_K2::PC_Object) && 
+			(EntryPin->PinType.PinCategory != UEdGraphSchema_K2::PC_Class) && 
+			(EntryPin->PinType.PinCategory != UEdGraphSchema_K2::PC_Interface) && 
 			!EntryPin->DefaultValue.IsEmpty() )
 		{
 			Function->SetMetaData(*EntryPin->PinName, *EntryPin->DefaultValue);
 		}
 	}
+
+	if(UFunction* OverriddenFunction = Function->GetSuperFunction())
+	{
+		// Copy metadata from parent function as well
+		UMetaData::CopyMetadata(OverriddenFunction, Function);
+	}
 }
 
 /**
@@ -2096,7 +2135,6 @@ void FKismetCompilerContext::FinishCompilingClass(UClass* Class)
 		Class->ClassFlags &= ~CLASS_RecompilerClear;
 		Class->ClassFlags |= (ParentClass->ClassFlags & CLASS_ScriptInherit);//@TODO: ChangeParentClass had this, but I don't think I want it: | UClass::StaticClassFlags;  // will end up with CLASS_Intrinsic
 		Class->ClassCastFlags |= ParentClass->ClassCastFlags;
-		Class->ClassWithin = ParentClass->ClassWithin ? ParentClass->ClassWithin : UObject::StaticClass();
 		Class->ClassConfigName = ParentClass->ClassConfigName;
 
 		// If the Blueprint was marked as deprecated, then flag the class as deprecated.
@@ -2413,7 +2451,7 @@ void FKismetCompilerContext::ExpandTimelineNodes(UEdGraph* SourceGraph)
 				// This might fail if this is the first compile after adding the timeline (property doesn't exist yet) - in that case, manually add the output pin
 				if (TimelineVarPin == NULL)
 				{
-					TimelineVarPin = GetTimelineNode->CreatePin(EGPD_Output, Schema->PC_Object, TEXT(""), UTimelineComponent::StaticClass(), false, false, TimelineNode->TimelineName.ToString());
+					TimelineVarPin = GetTimelineNode->CreatePin(EGPD_Output, Schema->PC_Object, FString(), UTimelineComponent::StaticClass(), TimelineNode->TimelineName.ToString());
 				}
 
 				if (bPlayPinConnected)
@@ -2503,10 +2541,15 @@ FPinConnectionResponse FKismetCompilerContext::CopyPinLinksToIntermediate(UEdGra
 }
 
 UK2Node_TemporaryVariable* FKismetCompilerContext::SpawnInternalVariable(UEdGraphNode* SourceNode, FString Category, FString SubCategory, UObject* SubcategoryObject, bool bIsArray, bool bIsSet, bool bIsMap, const FEdGraphTerminalType& ValueTerminalType)
+{
+	return SpawnInternalVariable(SourceNode, MoveTemp(Category), MoveTemp(SubCategory), SubcategoryObject, FEdGraphPinType::ToPinContainerType(bIsArray, bIsSet, bIsMap), ValueTerminalType);
+}
+
+UK2Node_TemporaryVariable* FKismetCompilerContext::SpawnInternalVariable(UEdGraphNode* SourceNode, FString Category, FString SubCategory, UObject* SubcategoryObject, EPinContainerType PinContainerType, const FEdGraphTerminalType& ValueTerminalType)
 {
 	UK2Node_TemporaryVariable* Result = SpawnIntermediateNode<UK2Node_TemporaryVariable>(SourceNode);
 
-	Result->VariableType = FEdGraphPinType(Category, SubCategory, SubcategoryObject, bIsArray, false, bIsSet, bIsMap, ValueTerminalType);
+	Result->VariableType = FEdGraphPinType(MoveTemp(Category), MoveTemp(SubCategory), SubcategoryObject, PinContainerType, false, ValueTerminalType);
 	Result->AllocateDefaultPins();
 
 	return Result;
@@ -3245,7 +3288,7 @@ void FKismetCompilerContext::ExpandTunnelsAndMacros(UEdGraph* SourceGraph)
 				if (Pin)
 				{
 					// Since we don't support array literals, drop a make array node on any unconnected array pins, which will allow macro expansion to succeed even if disconnected
-					if (Pin->PinType.bIsArray
+					if (Pin->PinType.IsArray()
 					&& (Pin->Direction == EGPD_Input)
 					&& (Pin->LinkedTo.Num() == 0))
 					{
@@ -3427,7 +3470,8 @@ void FKismetCompilerContext::ProcessOneFunctionGraph(UEdGraph* SourceGraph, bool
 	//       been added to the class yet, etc.)
 	if ((CompileOptions.CompileType == EKismetCompileType::SkeletonOnly) || ValidateGraphIsWellFormed(FunctionGraph))
 	{
-		FKismetFunctionContext& Context = *new (FunctionList)FKismetFunctionContext(MessageLog, Schema, NewClass, Blueprint, CompileOptions.DoesRequireCppCodeGeneration(), CompileOptions.IsInstrumentationActive());
+		const UEdGraphSchema_K2* FunctionGraphSchema = CastChecked<const UEdGraphSchema_K2>(FunctionGraph->GetSchema());
+		FKismetFunctionContext& Context = *new (FunctionList)FKismetFunctionContext(MessageLog, FunctionGraphSchema, NewClass, Blueprint, CompileOptions.DoesRequireCppCodeGeneration(), CompileOptions.IsInstrumentationActive());
 		Context.SourceGraph = FunctionGraph;
 
 		if(FBlueprintEditorUtils::IsDelegateSignatureGraph(SourceGraph))
@@ -3442,7 +3486,7 @@ void FKismetCompilerContext::ProcessOneFunctionGraph(UEdGraph* SourceGraph, bool
 		}
 
 		bool bEnforceConstCorrectness = true;
-		if (FBlueprintEditorUtils::IsBlueprintConst(Blueprint) || Schema->IsConstFunctionGraph(Context.SourceGraph, &bEnforceConstCorrectness))
+		if (FBlueprintEditorUtils::IsBlueprintConst(Blueprint) || Context.Schema->IsConstFunctionGraph(Context.SourceGraph, &bEnforceConstCorrectness))
 		{
 			Context.MarkAsConstFunction(bEnforceConstCorrectness);
 		}
@@ -3712,8 +3756,7 @@ void FKismetCompilerContext::CompileClassLayout(EInternalCompilerFlags InternalF
 	NewClass->ClassGeneratedBy = Blueprint;
 	
 	// Set class metadata as needed
-	UClass* ParentClass = NewClass->ClassWithin;
-	NewClass->SetSuperStruct(ParentClass);
+	UClass* ParentClass = NewClass->GetSuperClass();
 	NewClass->ClassFlags |= (ParentClass->ClassFlags & CLASS_Inherit);
 	NewClass->ClassCastFlags |= ParentClass->ClassCastFlags;
 	
@@ -3776,7 +3819,7 @@ void FKismetCompilerContext::CompileClassLayout(EInternalCompilerFlags InternalF
 	if (UsePersistentUberGraphFrame() && UbergraphContext)
 	{
 		//UBER GRAPH PERSISTENT FRAME
-		FEdGraphPinType Type(TEXT("struct"), TEXT(""), FPointerToUberGraphFrame::StaticStruct(), false, false, false, false, FEdGraphTerminalType());
+		FEdGraphPinType Type(TEXT("struct"), FString(), FPointerToUberGraphFrame::StaticStruct(), EPinContainerType::None, false, FEdGraphTerminalType());
 		UProperty* Property = CreateVariable(UBlueprintGeneratedClass::GetUberGraphFrameName(), Type);
 		Property->SetPropertyFlags(CPF_DuplicateTransient | CPF_Transient);
 	}
@@ -4255,7 +4298,7 @@ void FKismetCompilerContext::CompileFunctions(EInternalCompilerFlags InternalFla
 
 					if (UClass* AsClass = Cast<UClass>(Struct))
 					{
-						Ar << AsClass->ClassFlags;
+						Ar << (uint32&)AsClass->ClassFlags;
 						Ar << AsClass->Interfaces;
 					}
 
diff --git a/Engine/Source/Editor/KismetCompiler/Private/KismetCompilerMisc.cpp b/Engine/Source/Editor/KismetCompiler/Private/KismetCompilerMisc.cpp
--- a/Engine/Source/Editor/KismetCompiler/Private/KismetCompilerMisc.cpp
+++ b/Engine/Source/Editor/KismetCompiler/Private/KismetCompilerMisc.cpp
@@ -277,7 +277,7 @@ bool FKismetCompilerUtilities::IsTypeCompatibleWithProperty(UEdGraphPin* SourceP
 	int32 NumErrorsAtStart = MessageLog.NumErrors;
 	bool bTypeMismatch = false;
 
-	if( Type.bIsArray )
+	if( Type.IsArray() )
 	{
 		// For arrays, the property we want to test against is the inner property
 		if( UArrayProperty* ArrayProp = Cast<UArrayProperty>(Property) )
@@ -309,7 +309,7 @@ bool FKismetCompilerUtilities::IsTypeCompatibleWithProperty(UEdGraphPin* SourceP
 			return false;
 		}
 	}
-	else if (Type.bIsSet)
+	else if (Type.IsSet())
 	{
 		if (USetProperty* SetProperty = Cast<USetProperty>(Property))
 		{
@@ -326,7 +326,7 @@ bool FKismetCompilerUtilities::IsTypeCompatibleWithProperty(UEdGraphPin* SourceP
 			return false;
 		}
 	}
-	else if (Type.bIsMap)
+	else if (Type.IsMap())
 	{
 		if (UMapProperty* MapProperty = Cast<UMapProperty>(Property))
 		{
@@ -662,7 +662,7 @@ bool FKismetCompilerUtilities::ValidateSelfCompatibility(const UEdGraphPin* Pin,
 {
 	const UBlueprint* Blueprint = Context.Blueprint;
 	const UEdGraph* SourceGraph = Context.SourceGraph;
-	UEdGraphSchema_K2* K2Schema = Context.Schema;
+	const UEdGraphSchema_K2* K2Schema = Context.Schema;
 	const UBlueprintGeneratedClass* BPClass = Context.NewClass;
 
 	FString ErrorMsg;
@@ -757,8 +757,8 @@ UEdGraphPin* FKismetCompilerUtilities::GenerateAssignmentNodes(class FKismetComp
 			UFunction* SetByNameFunction = Schema->FindSetVariableByNameFunction(OrgPin->PinType);
 			if (SetByNameFunction)
 			{
-				UK2Node_CallFunction* SetVarNode = NULL;
-				if (OrgPin->PinType.bIsArray)
+				UK2Node_CallFunction* SetVarNode = nullptr;
+				if (OrgPin->PinType.IsArray())
 				{
 					SetVarNode = CompilerContext.SpawnIntermediateNode<UK2Node_CallArrayFunction>(SpawnNode, SourceGraph);
 				}
@@ -801,7 +801,7 @@ UEdGraphPin* FKismetCompilerUtilities::GenerateAssignmentNodes(class FKismetComp
 				else
 				{
 					// For non-array struct pins that are not linked, transfer the pin type so that the node will expand an auto-ref that will assign the value by-ref.
-					if (OrgPin->PinType.bIsArray == false && OrgPin->PinType.PinCategory == Schema->PC_Struct && OrgPin->LinkedTo.Num() == 0)
+					if (OrgPin->PinType.IsArray() == false && OrgPin->PinType.PinCategory == Schema->PC_Struct && OrgPin->LinkedTo.Num() == 0)
 					{
 						ValuePin->PinType.PinCategory = OrgPin->PinType.PinCategory;
 						ValuePin->PinType.PinSubCategory = OrgPin->PinType.PinSubCategory;
@@ -1085,14 +1085,10 @@ UProperty* FKismetCompilerUtilities::CreatePropertyOnScope(UStruct* Scope, const
 		}
 	}
 
-	// Handle creating an array property, if necessary
-	const bool bIsMapProperty = Type.bIsMap;
-	const bool bIsSetProperty = Type.bIsSet;
-	const bool bIsArrayProperty = Type.bIsArray;
-	// only one of these should be set (or none):
-	ensure( !(bIsMapProperty || bIsSetProperty || bIsArrayProperty) ||
-			( !(bIsMapProperty && bIsSetProperty && bIsArrayProperty) && (bIsMapProperty ^ bIsSetProperty ^ bIsArrayProperty) )
-			);
+	// Handle creating a container property, if necessary
+	const bool bIsMapProperty = Type.IsMap();
+	const bool bIsSetProperty = Type.IsSet();
+	const bool bIsArrayProperty = Type.IsArray();
 	UMapProperty* NewMapProperty = nullptr;
 	USetProperty* NewSetProperty = nullptr;
 	UArrayProperty* NewArrayProperty = nullptr;
@@ -1634,8 +1630,8 @@ FBPTerminal* FNodeHandlingFunctor::RegisterLiteral(FKismetFunctionContext& Conte
 		FString ErrorString = FText::Format(ErrorFormat, InvalidReasonText, DefaultValue).ToString();
 		CompilerContext.MessageLog.Error(*ErrorString, Net);
 
-		// Skip over these properties if they are array or ref properties, because the backend can't emit valid code for them
-		if (Net->PinType.bIsArray || Net->PinType.bIsReference)
+		// Skip over these properties if they are container or ref properties, because the backend can't emit valid code for them
+		if (Net->PinType.IsContainer() || Net->PinType.bIsReference)
 		{
 			return nullptr;
 		}
@@ -1699,7 +1695,7 @@ FString FNetNameMapping::MakeBaseName(const UAnimGraphNode_Base* Net)
 //////////////////////////////////////////////////////////////////////////
 // FKismetFunctionContext
 
-FKismetFunctionContext::FKismetFunctionContext(FCompilerResultsLog& InMessageLog, UEdGraphSchema_K2* InSchema, UBlueprintGeneratedClass* InNewClass, UBlueprint* InBlueprint, bool bInGeneratingCpp, bool bInWantsInstrumentation)
+FKismetFunctionContext::FKismetFunctionContext(FCompilerResultsLog& InMessageLog, const UEdGraphSchema_K2* InSchema, UBlueprintGeneratedClass* InNewClass, UBlueprint* InBlueprint, bool bInGeneratingCpp, bool bInWantsInstrumentation)
 	: Blueprint(InBlueprint)
 	, SourceGraph(nullptr)
 	, EntryPoint(nullptr)
@@ -2103,7 +2099,7 @@ struct FEventGraphUtils
 	{
 		// TODO: Strange cases..
 		if ((Net.Direction != EEdGraphPinDirection::EGPD_Output)
-			|| Net.PinType.bIsArray
+			|| Net.PinType.IsContainer()
 			|| Net.PinType.bIsReference
 			|| Net.PinType.bIsConst
 			|| Net.SubPins.Num())
diff --git a/Engine/Source/Editor/KismetCompiler/Private/KismetCompilerModule.cpp b/Engine/Source/Editor/KismetCompiler/Private/KismetCompilerModule.cpp
--- a/Engine/Source/Editor/KismetCompiler/Private/KismetCompilerModule.cpp
+++ b/Engine/Source/Editor/KismetCompiler/Private/KismetCompilerModule.cpp
@@ -49,8 +49,8 @@ public:
 	virtual void RemoveBlueprintGeneratedClasses(class UBlueprint* Blueprint) override;
 	virtual TArray<IBlueprintCompiler*>& GetCompilers() override { return Compilers; }
 	virtual void GetBlueprintTypesForClass(UClass* ParentClass, UClass*& OutBlueprintClass, UClass*& OutBlueprintGeneratedClass) const override;
-	virtual void GenerateCppCodeForEnum(UUserDefinedEnum* UDEnum, FString& OutHeaderCode, FString& OutCPPCode) override;
-	virtual FString GenerateCppCodeForStruct(UUserDefinedStruct* UDStruct, const FCompilerNativizationOptions& NativizationOptions) override;
+	virtual void GenerateCppCodeForEnum(UUserDefinedEnum* UDEnum, const FCompilerNativizationOptions& NativizationOptions, FString& OutHeaderCode, FString& OutCPPCode) override;
+	virtual void GenerateCppCodeForStruct(UUserDefinedStruct* UDStruct, const FCompilerNativizationOptions& NativizationOptions, FString& OutHeaderCode, FString& OutCPPCode) override;
 	virtual FString GenerateCppWrapper(UBlueprintGeneratedClass* BPGC, const FCompilerNativizationOptions& NativizationOptions) override;
 	// End implementation
 private:
@@ -179,16 +179,16 @@ void FKismet2CompilerModule::CompileStructure(UUserDefinedStruct* Struct, FCompi
 	FUserDefinedStructureCompilerUtils::CompileStruct(Struct, Results, true);
 }
 
-void FKismet2CompilerModule::GenerateCppCodeForEnum(UUserDefinedEnum* UDEnum, FString& OutHeaderCode, FString& OutCPPCode)
+void FKismet2CompilerModule::GenerateCppCodeForEnum(UUserDefinedEnum* UDEnum, const FCompilerNativizationOptions& NativizationOptions, FString& OutHeaderCode, FString& OutCPPCode)
 {
 	TUniquePtr<IBlueprintCompilerCppBackend> Backend_CPP(IBlueprintCompilerCppBackendModuleInterface::Get().Create());
-	Backend_CPP->GenerateCodeFromEnum(UDEnum, OutHeaderCode, OutCPPCode);
+	Backend_CPP->GenerateCodeFromEnum(UDEnum, NativizationOptions, OutHeaderCode, OutCPPCode);
 }
 
-FString FKismet2CompilerModule::GenerateCppCodeForStruct(UUserDefinedStruct* UDStruct, const FCompilerNativizationOptions& NativizationOptions)
+void FKismet2CompilerModule::GenerateCppCodeForStruct(UUserDefinedStruct* UDStruct, const FCompilerNativizationOptions& NativizationOptions, FString& OutHeaderCode, FString& OutCPPCode)
 {
 	TUniquePtr<IBlueprintCompilerCppBackend> Backend_CPP(IBlueprintCompilerCppBackendModuleInterface::Get().Create());
-	return Backend_CPP->GenerateCodeFromStruct(UDStruct, NativizationOptions);
+	Backend_CPP->GenerateCodeFromStruct(UDStruct, NativizationOptions, OutHeaderCode, OutCPPCode);
 }
 
 FString FKismet2CompilerModule::GenerateCppWrapper(UBlueprintGeneratedClass* BPGC, const FCompilerNativizationOptions& NativizationOptions)
diff --git a/Engine/Source/Editor/KismetCompiler/Private/KismetCompilerVMBackend.cpp b/Engine/Source/Editor/KismetCompiler/Private/KismetCompilerVMBackend.cpp
--- a/Engine/Source/Editor/KismetCompiler/Private/KismetCompilerVMBackend.cpp
+++ b/Engine/Source/Editor/KismetCompiler/Private/KismetCompilerVMBackend.cpp
@@ -595,9 +595,7 @@ public:
 					    {
 						    return false;
 					    }
-					    if ((TypeA.bIsMap != TypeB.bIsMap)
-						    || (TypeA.bIsSet != TypeB.bIsSet)
-						    || (TypeA.bIsArray != TypeB.bIsArray)
+					    if ((TypeA.ContainerType != TypeB.ContainerType)
 						    || (TypeA.bIsWeakPointer != TypeB.bIsWeakPointer))
 					    {
 						    return false;
@@ -791,10 +789,13 @@ public:
 				if (Struct == VectorStruct)
 				{
 					FVector V = FVector::ZeroVector;
-					const bool bParsedUsingCustomFormat = FDefaultValueHelper::ParseVector(Term->Name, /*out*/ V);
-					if (!bParsedUsingCustomFormat)
+					if (!Term->Name.IsEmpty())
 					{
-						Struct->ImportText(*Term->Name, &V, nullptr, PPF_None, GWarn, GetPathNameSafe(StructProperty));
+						const bool bParsedUsingCustomFormat = FDefaultValueHelper::ParseVector(Term->Name, /*out*/ V);
+						if (!bParsedUsingCustomFormat)
+						{
+							Struct->ImportText(*Term->Name, &V, nullptr, PPF_None, GWarn, GetPathNameSafe(StructProperty));
+						}
 					}
 					Writer << EX_VectorConst;
 					Writer << V;
@@ -802,10 +803,13 @@ public:
 				else if (Struct == RotatorStruct)
 				{
 					FRotator R = FRotator::ZeroRotator;
-					const bool bParsedUsingCustomFormat = FDefaultValueHelper::ParseRotator(Term->Name, /*out*/ R);
-					if (!bParsedUsingCustomFormat)
+					if (!Term->Name.IsEmpty())
 					{
-						Struct->ImportText(*Term->Name, &R, nullptr, PPF_None, GWarn, GetPathNameSafe(StructProperty));
+						const bool bParsedUsingCustomFormat = FDefaultValueHelper::ParseRotator(Term->Name, /*out*/ R);
+						if (!bParsedUsingCustomFormat)
+						{
+							Struct->ImportText(*Term->Name, &R, nullptr, PPF_None, GWarn, GetPathNameSafe(StructProperty));
+						}
 					}
 					Writer << EX_RotationConst;
 					Writer << R;
@@ -968,8 +972,7 @@ public:
 			else if (FLiteralTypeHelper::IsAsset(&Term->Type, CoerceProperty))
 			{
 				Writer << EX_AssetConst;
-				FAssetPtr AssetPtr(Term->ObjectLiteral);
-				EmitStringLiteral(AssetPtr.GetUniqueID().ToString());
+				EmitStringLiteral(Term->Name);
 			}
 			else if (FLiteralTypeHelper::IsObject(&Term->Type, CoerceProperty) || FLiteralTypeHelper::IsClass(&Term->Type, CoerceProperty))
 			{
@@ -1343,7 +1346,7 @@ public:
 	{
 		check(Schema && DestinationExpression && !DestinationExpression->Type.PinCategory.IsEmpty());
 
-		const bool bIsContainer = DestinationExpression->Type.bIsArray || DestinationExpression->Type.bIsSet || DestinationExpression->Type.bIsMap;
+		const bool bIsContainer = DestinationExpression->Type.IsContainer();
 		const bool bIsDelegate = Schema->PC_Delegate == DestinationExpression->Type.PinCategory;
 		const bool bIsMulticastDelegate = Schema->PC_MCDelegate == DestinationExpression->Type.PinCategory;
 		const bool bIsBoolean = Schema->PC_Boolean == DestinationExpression->Type.PinCategory;
diff --git a/Engine/Source/Editor/KismetCompiler/Public/BlueprintCompilerCppBackendInterface.h b/Engine/Source/Editor/KismetCompiler/Public/BlueprintCompilerCppBackendInterface.h
--- a/Engine/Source/Editor/KismetCompiler/Public/BlueprintCompilerCppBackendInterface.h
+++ b/Engine/Source/Editor/KismetCompiler/Public/BlueprintCompilerCppBackendInterface.h
@@ -15,8 +15,8 @@ class IBlueprintCompilerCppBackend
 {
 public:
 	virtual FString GenerateCodeFromClass(UClass* SourceClass, TIndirectArray<FKismetFunctionContext>& Functions, bool bGenerateStubsOnly, const FCompilerNativizationOptions& NativizationOptions, FString& OutCppBody) = 0;
-	virtual void GenerateCodeFromEnum(UUserDefinedEnum* SourceEnum, FString& OutHeaderCode, FString& OutCPPCode) = 0;
-	virtual FString GenerateCodeFromStruct(UUserDefinedStruct* SourceStruct, const FCompilerNativizationOptions& NativizationOptions) = 0;
+	virtual void GenerateCodeFromEnum(UUserDefinedEnum* SourceEnum, const FCompilerNativizationOptions& NativizationOptions, FString& OutHeaderCode, FString& OutCPPCode) = 0;
+	virtual void GenerateCodeFromStruct(UUserDefinedStruct* SourceStruct, const FCompilerNativizationOptions& NativizationOptions, FString& OutHeaderCode, FString& OutCPPCode) = 0;
 
 	// Generate a wrapper class, that helps accessing non-native properties and calling non-native functions
 	virtual FString GenerateWrapperForClass(UClass* SourceClass, const FCompilerNativizationOptions& NativizationOptions) = 0;
diff --git a/Engine/Source/Editor/KismetCompiler/Public/KismetCompiledFunctionContext.h b/Engine/Source/Editor/KismetCompiler/Public/KismetCompiledFunctionContext.h
--- a/Engine/Source/Editor/KismetCompiler/Public/KismetCompiledFunctionContext.h
+++ b/Engine/Source/Editor/KismetCompiler/Public/KismetCompiledFunctionContext.h
@@ -54,7 +54,7 @@ public:
 	TArray<UEdGraphNode*> LinearExecutionList;
 
 	FCompilerResultsLog& MessageLog;
-	UEdGraphSchema_K2* Schema;
+	const UEdGraphSchema_K2* Schema;
 
 	// An UNORDERED listing of all statements (used for cleaning up the dynamically allocated statements)
 	TArray< FBlueprintCompiledStatement* > AllGeneratedStatements;
@@ -105,7 +105,7 @@ public:
 	//Does this function use requires FlowStack ?
 	bool bUseFlowStack;
 public:
-	FKismetFunctionContext(FCompilerResultsLog& InMessageLog, UEdGraphSchema_K2* InSchema, UBlueprintGeneratedClass* InNewClass, UBlueprint* InBlueprint, bool bInGeneratingCpp, bool bInWantsInstrumentation);
+	FKismetFunctionContext(FCompilerResultsLog& InMessageLog, const UEdGraphSchema_K2* InSchema, UBlueprintGeneratedClass* InNewClass, UBlueprint* InBlueprint, bool bInGeneratingCpp, bool bInWantsInstrumentation);
 	
 	~FKismetFunctionContext();
 
diff --git a/Engine/Source/Editor/KismetCompiler/Public/KismetCompiler.h b/Engine/Source/Editor/KismetCompiler/Public/KismetCompiler.h
--- a/Engine/Source/Editor/KismetCompiler/Public/KismetCompiler.h
+++ b/Engine/Source/Editor/KismetCompiler/Public/KismetCompiler.h
@@ -216,7 +216,10 @@ public:
 	 */
 	FPinConnectionResponse CopyPinLinksToIntermediate(UEdGraphPin& SourcePin, UEdGraphPin& IntermediatePin);
 
-	UK2Node_TemporaryVariable* SpawnInternalVariable(UEdGraphNode* SourceNode, FString Category, FString SubCategory = TEXT(""), UObject* SubcategoryObject = NULL, bool bIsArray = false, bool bIsSet = false, bool bIsMap = false, const FEdGraphTerminalType& ValueTerminalType = FEdGraphTerminalType());
+	DEPRECATED(4.17, "Use version that takes PinContainerType instead of separate booleans for array, set, and map")
+	UK2Node_TemporaryVariable* SpawnInternalVariable(UEdGraphNode* SourceNode, FString Category, FString SubCategory, UObject* SubcategoryObject, bool bIsArray, bool bIsSet = false, bool bIsMap = false, const FEdGraphTerminalType& ValueTerminalType = FEdGraphTerminalType());
+
+	UK2Node_TemporaryVariable* SpawnInternalVariable(UEdGraphNode* SourceNode, FString Category, FString SubCategory = FString(), UObject* SubcategoryObject = nullptr, EPinContainerType PinContainerType = EPinContainerType::None, const FEdGraphTerminalType& ValueTerminalType = FEdGraphTerminalType());
 
 	bool UsePersistentUberGraphFrame() const;
 
@@ -332,6 +335,9 @@ protected:
 	/** Creates user defined local variables for function */
 	void CreateUserDefinedLocalVariablesForFunction(FKismetFunctionContext& Context, UField**& FunctionPropertyStorageLocation);
 
+	/** Helper function for CreateUserDefinedLocalVariablesForFunction and compilation manager's FastGenerateSkeletonClass: */
+	static UProperty* CreateUserDefinedLocalVariableForFunction(const FBPVariableDescription& Variable, UFunction* Function, UBlueprintGeneratedClass* OwningClass, UField**& FunctionPropertyStorageLocation, const UEdGraphSchema_K2* Schema, FCompilerResultsLog& MessageLog);
+
 	/** Adds a default value entry into the DefaultPropertyValueMap for the property specified */
 	void SetPropertyDefaultValue(const UProperty* PropertyToSet, FString& Value);
 
@@ -428,6 +434,8 @@ protected:
 	 */
 	void FinishCompilingFunction(FKismetFunctionContext& Context);
 
+	static void SetCalculatedMetaDataAndFlags(UFunction* Function, UK2Node_FunctionEntry* EntryNode, const UEdGraphSchema_K2* Schema );
+
 	/**
 	 * Handles adding the implemented interface information to the class
 	 */
diff --git a/Engine/Source/Editor/KismetCompiler/Public/KismetCompilerModule.h b/Engine/Source/Editor/KismetCompiler/Public/KismetCompilerModule.h
--- a/Engine/Source/Editor/KismetCompiler/Public/KismetCompilerModule.h
+++ b/Engine/Source/Editor/KismetCompiler/Public/KismetCompilerModule.h
@@ -100,8 +100,8 @@ public:
 	 */
 	virtual void GetBlueprintTypesForClass(UClass* ParentClass, UClass*& OutBlueprintClass, UClass*& OutBlueprintGeneratedClass) const = 0;
 
-	virtual void GenerateCppCodeForEnum(UUserDefinedEnum* UDEnum, FString& OutHeaderCode, FString& OutCPPCode) = 0;
-	virtual FString GenerateCppCodeForStruct(UUserDefinedStruct* UDStruct, const FCompilerNativizationOptions& NativizationOptions) = 0;
+	virtual void GenerateCppCodeForEnum(UUserDefinedEnum* UDEnum, const FCompilerNativizationOptions& NativizationOptions, FString& OutHeaderCode, FString& OutCPPCode) = 0;
+	virtual void GenerateCppCodeForStruct(UUserDefinedStruct* UDStruct, const FCompilerNativizationOptions& NativizationOptions, FString& OutHeaderCode, FString& OutCPPCode) = 0;
 	// Generate a wrapper class, that helps accessing non-native properties and calling non-native functions
 	virtual FString GenerateCppWrapper(UBlueprintGeneratedClass* BPGC, const FCompilerNativizationOptions& NativizationOptions) = 0;
 };
diff --git a/Engine/Source/Editor/KismetWidgets/Private/SPinTypeSelector.cpp b/Engine/Source/Editor/KismetWidgets/Private/SPinTypeSelector.cpp
--- a/Engine/Source/Editor/KismetWidgets/Private/SPinTypeSelector.cpp
+++ b/Engine/Source/Editor/KismetWidgets/Private/SPinTypeSelector.cpp
@@ -131,14 +131,6 @@ private:
 	TWeakPtr<SSubMenuHandler> SubMenuHandler;
 };
 
-enum class EPinContainerType
-{
-	SingleValue,
-	Array,
-	Set,
-	Map,
-};
-
 static bool ContainerRequiresGetTypeHash(EPinContainerType InType)
 {
 	return InType == EPinContainerType::Set || InType == EPinContainerType::Map;
@@ -234,7 +226,7 @@ void SPinTypeSelector::Construct(const FArguments& InArgs, FGetPinTypeTree GetPi
 	static TArray<TSharedPtr<EPinContainerType>> PinTypes;
 	if (PinTypes.Num() == 0)
 	{
-		PinTypes.Add(MakeShareable(new EPinContainerType(EPinContainerType::SingleValue)));
+		PinTypes.Add(MakeShareable(new EPinContainerType(EPinContainerType::None)));
 		PinTypes.Add(MakeShareable(new EPinContainerType(EPinContainerType::Array)));
 		PinTypes.Add(MakeShareable(new EPinContainerType(EPinContainerType::Set)));
 		PinTypes.Add(MakeShareable(new EPinContainerType(EPinContainerType::Map)));
@@ -419,12 +411,12 @@ void SPinTypeSelector::Construct(const FArguments& InArgs, FGetPinTypeTree GetPi
 			.Visibility(
 				TAttribute<EVisibility>::Create(
 					TAttribute<EVisibility>::FGetter::CreateLambda(
-						[this]() {return this->TargetPinType.Get().bIsMap == true ? EVisibility::Visible : EVisibility::Collapsed; }
+						[this]() {return this->TargetPinType.Get().IsMap() == true ? EVisibility::Visible : EVisibility::Collapsed; }
 					)
 				)
 			)
 			[
-				SNew( SComboButton )
+				SAssignNew( SecondaryTypeComboButton, SComboButton )
 				.OnGetMenuContent(this, &SPinTypeSelector::GetMenuContent, true )
 				.ContentPadding(0)
 				.ToolTipText(this, &SPinTypeSelector::GetToolTipForComboBoxSecondaryType)
@@ -521,16 +513,13 @@ FSlateColor SPinTypeSelector::GetSecondaryTypeIconColor() const
 
 ECheckBoxState SPinTypeSelector::IsArrayChecked() const
 {
-	return TargetPinType.Get().bIsArray ? ECheckBoxState::Checked : ECheckBoxState::Unchecked;
+	return TargetPinType.Get().IsArray() ? ECheckBoxState::Checked : ECheckBoxState::Unchecked;
 }
 
 void SPinTypeSelector::OnArrayCheckStateChanged(ECheckBoxState NewState)
 {
 	FEdGraphPinType NewTargetPinType = TargetPinType.Get();
-	NewTargetPinType.bIsArray = (NewState == ECheckBoxState::Checked) ? true : false;
-	// just in case someone has been experimenting with experimental features:
-	NewTargetPinType.bIsMap = false;
-	NewTargetPinType.bIsSet = false;
+	NewTargetPinType.ContainerType = (NewState == ECheckBoxState::Checked) ? EPinContainerType::Array : EPinContainerType::None;
 
 	OnTypeChanged.ExecuteIfBound(NewTargetPinType);
 }
@@ -543,30 +532,7 @@ void SPinTypeSelector::OnArrayStateToggled()
 void SPinTypeSelector::OnContainerTypeSelectionChanged( EPinContainerType PinContainerType)
 {
 	FEdGraphPinType NewTargetPinType = TargetPinType.Get();
-
-	switch (PinContainerType)
-	{
-	case EPinContainerType::SingleValue:
-		NewTargetPinType.bIsArray = false;
-		NewTargetPinType.bIsMap = false;
-		NewTargetPinType.bIsSet = false;
-		break;
-	case EPinContainerType::Array:
-		NewTargetPinType.bIsArray = true;
-		NewTargetPinType.bIsMap = false;
-		NewTargetPinType.bIsSet = false;
-		break;
-	case EPinContainerType::Set:
-		NewTargetPinType.bIsArray = false;
-		NewTargetPinType.bIsMap = false;
-		NewTargetPinType.bIsSet = true;
-		break;
-	case EPinContainerType::Map:
-		NewTargetPinType.bIsArray = false;
-		NewTargetPinType.bIsMap = true;
-		NewTargetPinType.bIsSet = false;
-		break;
-	}
+	NewTargetPinType.ContainerType = PinContainerType;
 
 	OnTypeChanged.ExecuteIfBound(NewTargetPinType);
 }
@@ -817,26 +783,27 @@ void SPinTypeSelector::OnSelectPinType(FPinTypeTreeItem InItem, FString InPinCat
 
 
 	TypeComboButton->SetIsOpen(false);
+	if (SecondaryTypeComboButton.IsValid())
+	{
+		SecondaryTypeComboButton->SetIsOpen(false);
+	}
 
 	if( NewTargetPinType.PinCategory == Schema->PC_Exec )
 	{
-		NewTargetPinType.bIsMap = false;
-		NewTargetPinType.bIsSet = false;
-		NewTargetPinType.bIsArray = false;
+		NewTargetPinType.ContainerType = EPinContainerType::None;
 		NewTargetPinType.PinValueType.TerminalCategory = FString();
 		NewTargetPinType.PinValueType.TerminalSubCategory = FString();
 		NewTargetPinType.PinValueType.TerminalSubCategoryObject = nullptr;
 	}
 
-	if ((NewTargetPinType.bIsMap || NewTargetPinType.bIsSet) && !FBlueprintEditorUtils::HasGetTypeHash(NewTargetPinType))
+	if ((NewTargetPinType.IsMap() || NewTargetPinType.IsSet()) && !FBlueprintEditorUtils::HasGetTypeHash(NewTargetPinType))
 	{
 		FEdGraphPinType HashedType = NewTargetPinType;
 		// clear the container-ness for messaging, we want to explain that the contained type is not hashable,
 		// not message about the container type (e.g. "Container type cleared because 'bool' does not have a GetTypeHash..."
 		// instead of "Container Type cleared because 'map of bool to float'..."). We also need to clear this because
 		// the type cannot be a container:
-		NewTargetPinType.bIsMap = false;
-		NewTargetPinType.bIsSet = false;
+		NewTargetPinType.ContainerType = EPinContainerType::None;
 
 		// inform user via toast why the type change was exceptional and clear IsMap/IsSetness because this type cannot be hashed:
 		const FText NotificationText = FText::Format(LOCTEXT("TypeCannotBeHashed", "Container type cleared because '{0}' does not have a GetTypeHash function. Maps and Sets require a hash function to insert and find elements"), UEdGraphSchema_K2::TypeToText(NewTargetPinType));
@@ -965,6 +932,10 @@ TSharedRef<SWidget>	SPinTypeSelector::GetMenuContent(bool bForSecondaryType)
 			
 
 			TypeComboButton->SetMenuContentWidgetToFocus(FilterTextBox);
+			if (SecondaryTypeComboButton.IsValid())
+			{
+				SecondaryTypeComboButton->SetMenuContentWidgetToFocus(FilterTextBox);
+			}
 	}
 	else
 	{
diff --git a/Engine/Source/Editor/KismetWidgets/Public/SPinTypeSelector.h b/Engine/Source/Editor/KismetWidgets/Public/SPinTypeSelector.h
--- a/Engine/Source/Editor/KismetWidgets/Public/SPinTypeSelector.h
+++ b/Engine/Source/Editor/KismetWidgets/Public/SPinTypeSelector.h
@@ -48,7 +48,7 @@ public:
 
 	SLATE_BEGIN_ARGS( SPinTypeSelector )
 		: _TargetPinType()
-		, _Schema(NULL)
+		, _Schema(nullptr)
 		, _TypeTreeFilter(ETypeTreeFilter::None)
 		, _bAllowArrays(true)
 		, _TreeViewWidth(300.f)
@@ -96,6 +96,7 @@ protected:
 	FText GetSecondaryTypeDescription() const;
 
 	TSharedPtr<SComboButton>		TypeComboButton;
+	TSharedPtr<SComboButton>		SecondaryTypeComboButton;
 	TSharedPtr<SSearchBox>			FilterTextBox;
 	TSharedPtr<SPinTypeTreeView>	TypeTreeView;
 	
@@ -140,7 +141,7 @@ protected:
 	void OnArrayStateToggled();
 
 	/** Updates the variable container type: */
-	void OnContainerTypeSelectionChanged(enum class EPinContainerType PinContainerType);
+	void OnContainerTypeSelectionChanged(EPinContainerType PinContainerType);
 
 	/** Array containing the unfiltered list of all supported types this pin could possibly have */
 	TArray<FPinTypeTreeItem>		TypeTreeRoot;
diff --git a/Engine/Source/Editor/LandscapeEditor/Private/LandscapeEdModeSplineTools.cpp b/Engine/Source/Editor/LandscapeEditor/Private/LandscapeEdModeSplineTools.cpp
--- a/Engine/Source/Editor/LandscapeEditor/Private/LandscapeEdModeSplineTools.cpp
+++ b/Engine/Source/Editor/LandscapeEditor/Private/LandscapeEdModeSplineTools.cpp
@@ -657,7 +657,7 @@ public:
 					ToLandscape->SplineComponent->Modify();
 
 					const FTransform OldToNewTransform =
-						FromProxy->SplineComponent->ComponentToWorld.GetRelativeTransform(ToLandscape->SplineComponent->ComponentToWorld);
+						FromProxy->SplineComponent->GetComponentTransform().GetRelativeTransform(ToLandscape->SplineComponent->GetComponentTransform());
 
 					if (FromProxies.Find(FromProxy) == nullptr)
 					{
@@ -862,7 +862,7 @@ public:
 				SplinesComponent = Landscape->SplineComponent;
 			}
 
-			const FTransform LandscapeToSpline = Landscape->LandscapeActorToWorld().GetRelativeTransform(SplinesComponent->ComponentToWorld);
+			const FTransform LandscapeToSpline = Landscape->LandscapeActorToWorld().GetRelativeTransform(SplinesComponent->GetComponentTransform());
 
 			AddControlPoint(SplinesComponent, LandscapeToSpline.TransformPosition(InHitLocation));
 
@@ -1172,7 +1172,7 @@ public:
 						ClickedSplineSegment = SplineProxy->SplineSegment;
 						ALandscapeProxy* LandscapeProxy = ClickedSplineSegment->GetTypedOuter<ALandscapeProxy>();
 						check(LandscapeProxy);
-						LandscapeToSpline = LandscapeProxy->LandscapeActorToWorld().GetRelativeTransform(ClickedSplineSegment->GetOuterULandscapeSplinesComponent()->ComponentToWorld);
+						LandscapeToSpline = LandscapeProxy->LandscapeActorToWorld().GetRelativeTransform(ClickedSplineSegment->GetOuterULandscapeSplinesComponent()->GetComponentTransform());
 					}
 					else if (HitProxy->IsA(HActor::StaticGetType()))
 					{
@@ -1191,7 +1191,7 @@ public:
 									{
 										ClickedSplineSegment = SplineSegment;
 										ALandscapeProxy* LandscapeProxy = CastChecked<ALandscapeProxy>(SplineComponent->GetOwner());
-										LandscapeToSpline = LandscapeProxy->LandscapeActorToWorld().GetRelativeTransform(SplineComponent->ComponentToWorld);
+										LandscapeToSpline = LandscapeProxy->LandscapeActorToWorld().GetRelativeTransform(SplineComponent->GetComponentTransform());
 									}
 								}
 							}
@@ -1300,7 +1300,7 @@ public:
 			Connection.ControlPoint->GetConnectionLocationAndRotation(Connection.SocketName, StartLocation, StartRotation);
 
 			float OldTangentLen = Connection.TangentLen;
-			Connection.TangentLen += SplinesComponent->ComponentToWorld.InverseTransformVector(-Drag) | StartRotation.Vector();
+			Connection.TangentLen += SplinesComponent->GetComponentTransform().InverseTransformVector(-Drag) | StartRotation.Vector();
 
 			// Disallow a tangent of exactly 0
 			if (Connection.TangentLen == 0)
@@ -1332,11 +1332,11 @@ public:
 			{
 				const ULandscapeSplinesComponent* SplinesComponent = ControlPoint->GetOuterULandscapeSplinesComponent();
 
-				ControlPoint->Location += SplinesComponent->ComponentToWorld.InverseTransformVector(Drag);
+				ControlPoint->Location += SplinesComponent->GetComponentTransform().InverseTransformVector(Drag);
 
 				FVector RotAxis; float RotAngle;
 				InRot.Quaternion().ToAxisAndAngle(RotAxis, RotAngle);
-				RotAxis = (SplinesComponent->ComponentToWorld.GetRotation().Inverse() * ControlPoint->Rotation.Quaternion().Inverse()).RotateVector(RotAxis);
+				RotAxis = (SplinesComponent->GetComponentTransform().GetRotation().Inverse() * ControlPoint->Rotation.Quaternion().Inverse()).RotateVector(RotAxis);
 
 				// Hack: for some reason FQuat.Rotator() Clamps to 0-360 range, so use .GetNormalized() to recover the original negative rotation.
 				ControlPoint->Rotation += FQuat(RotAxis, RotAngle).Rotator().GetNormalized();
@@ -1450,8 +1450,8 @@ public:
 			{
 				const ULandscapeSplinesComponent* SplinesComponent = ControlPoint->GetOuterULandscapeSplinesComponent();
 
-				FVector HandlePos0 = SplinesComponent->ComponentToWorld.TransformPosition(ControlPoint->Location + ControlPoint->Rotation.Vector() * -20);
-				FVector HandlePos1 = SplinesComponent->ComponentToWorld.TransformPosition(ControlPoint->Location + ControlPoint->Rotation.Vector() * 20);
+				FVector HandlePos0 = SplinesComponent->GetComponentTransform().TransformPosition(ControlPoint->Location + ControlPoint->Rotation.Vector() * -20);
+				FVector HandlePos1 = SplinesComponent->GetComponentTransform().TransformPosition(ControlPoint->Location + ControlPoint->Rotation.Vector() * 20);
 				DrawDashedLine(PDI, HandlePos0, HandlePos1, FColor::White, 20, SDPG_Foreground);
 
 				if (GLevelEditorModeTools().GetWidgetMode() == FWidget::WM_Scale)
@@ -1461,8 +1461,8 @@ public:
 						FVector StartLocation; FRotator StartRotation;
 						Connection.GetNearConnection().ControlPoint->GetConnectionLocationAndRotation(Connection.GetNearConnection().SocketName, StartLocation, StartRotation);
 
-						FVector StartPos = SplinesComponent->ComponentToWorld.TransformPosition(StartLocation);
-						FVector HandlePos = SplinesComponent->ComponentToWorld.TransformPosition(StartLocation + StartRotation.Vector() * Connection.GetNearConnection().TangentLen / 2);
+						FVector StartPos = SplinesComponent->GetComponentTransform().TransformPosition(StartLocation);
+						FVector HandlePos = SplinesComponent->GetComponentTransform().TransformPosition(StartLocation + StartRotation.Vector() * Connection.GetNearConnection().TangentLen / 2);
 						PDI->DrawLine(StartPos, HandlePos, FColor::White, SDPG_Foreground);
 
 						if (PDI->IsHitTesting()) PDI->SetHitProxy(new HLandscapeSplineProxy_Tangent(Connection.Segment, Connection.End));
@@ -1484,8 +1484,8 @@ public:
 						FVector StartLocation; FRotator StartRotation;
 						Connection.ControlPoint->GetConnectionLocationAndRotation(Connection.SocketName, StartLocation, StartRotation);
 
-						FVector EndPos = SplinesComponent->ComponentToWorld.TransformPosition(StartLocation);
-						FVector EndHandlePos = SplinesComponent->ComponentToWorld.TransformPosition(StartLocation + StartRotation.Vector() * Connection.TangentLen / 2);
+						FVector EndPos = SplinesComponent->GetComponentTransform().TransformPosition(StartLocation);
+						FVector EndHandlePos = SplinesComponent->GetComponentTransform().TransformPosition(StartLocation + StartRotation.Vector() * Connection.TangentLen / 2);
 
 						PDI->DrawLine(EndPos, EndHandlePos, FColor::White, SDPG_Foreground);
 						if (PDI->IsHitTesting()) PDI->SetHitProxy(new HLandscapeSplineProxy_Tangent(Segment, !!End));
@@ -1560,7 +1560,7 @@ public:
 			{
 				ULandscapeSplineControlPoint* FirstPoint = *SelectedSplineControlPoints.CreateConstIterator();
 				ULandscapeSplinesComponent* SplinesComponent = FirstPoint->GetOuterULandscapeSplinesComponent();
-				return SplinesComponent->ComponentToWorld.TransformPosition(FirstPoint->Location);
+				return SplinesComponent->GetComponentTransform().TransformPosition(FirstPoint->Location);
 			}
 		}
 
@@ -1576,7 +1576,7 @@ public:
 			{
 				ULandscapeSplineControlPoint* FirstPoint = *SelectedSplineControlPoints.CreateConstIterator();
 				ULandscapeSplinesComponent* SplinesComponent = FirstPoint->GetOuterULandscapeSplinesComponent();
-				return FQuatRotationTranslationMatrix(FirstPoint->Rotation.Quaternion() * SplinesComponent->ComponentToWorld.GetRotation(), FVector::ZeroVector);
+				return FQuatRotationTranslationMatrix(FirstPoint->Rotation.Quaternion() * SplinesComponent->GetComponentTransform().GetRotation(), FVector::ZeroVector);
 			}
 		}
 
diff --git a/Engine/Source/Editor/MaterialEditor/Private/FindInMaterial.cpp b/Engine/Source/Editor/MaterialEditor/Private/FindInMaterial.cpp
--- a/Engine/Source/Editor/MaterialEditor/Private/FindInMaterial.cpp
+++ b/Engine/Source/Editor/MaterialEditor/Private/FindInMaterial.cpp
@@ -71,7 +71,7 @@ TSharedRef<SWidget> FFindInMaterialResult::CreateIcon() const
 	const FSlateBrush* Brush = NULL;
 	if (UEdGraphPin* ResolvedPin = Pin.Get())
 	{
-		if (ResolvedPin->PinType.bIsArray)
+		if (ResolvedPin->PinType.IsArray())
 		{
 			Brush = FEditorStyle::GetBrush(TEXT("GraphEditor.ArrayPinIcon"));
 		}
diff --git a/Engine/Source/Editor/MovieSceneTools/Private/K2Node_GetSequenceBinding.cpp b/Engine/Source/Editor/MovieSceneTools/Private/K2Node_GetSequenceBinding.cpp
--- a/Engine/Source/Editor/MovieSceneTools/Private/K2Node_GetSequenceBinding.cpp
+++ b/Engine/Source/Editor/MovieSceneTools/Private/K2Node_GetSequenceBinding.cpp
@@ -88,6 +88,8 @@ UMovieSceneSequence* UK2Node_GetSequenceBinding::GetSequence() const
 
 void UK2Node_GetSequenceBinding::ValidateNodeDuringCompilation(FCompilerResultsLog& MessageLog) const
 {
+	Super::ValidateNodeDuringCompilation(MessageLog);
+
 	UMovieScene* MovieScene = GetObjectMovieScene();
 	if (!MovieScene)
 	{
@@ -105,10 +107,10 @@ void UK2Node_GetSequenceBinding::AllocateDefaultPins()
 {
 	const UEdGraphSchema_K2* K2Schema = GetDefault<UEdGraphSchema_K2>();
 
-	CreatePin(EGPD_Input, K2Schema->PC_Object, TEXT(""), UMovieSceneSequence::StaticClass(), false, false, SequencePinName);
+	CreatePin(EGPD_Input, K2Schema->PC_Object, FString(), UMovieSceneSequence::StaticClass(), SequencePinName);
 
 	// Result pin
-	UEdGraphPin* ResultPin = CreatePin(EGPD_Output, K2Schema->PC_Struct, TEXT(""), FMovieSceneObjectBindingID::StaticStruct(), false, false, K2Schema->PN_ReturnValue);
+	UEdGraphPin* ResultPin = CreatePin(EGPD_Output, K2Schema->PC_Struct, FString(), FMovieSceneObjectBindingID::StaticStruct(), K2Schema->PN_ReturnValue);
 	ResultPin->PinFriendlyName = LOCTEXT("SequenceBindingOutput", "Binding");
 
 	Super::AllocateDefaultPins();
diff --git a/Engine/Source/Editor/MovieSceneTools/Private/K2Node_GetSequenceBindings.cpp b/Engine/Source/Editor/MovieSceneTools/Private/K2Node_GetSequenceBindings.cpp
--- a/Engine/Source/Editor/MovieSceneTools/Private/K2Node_GetSequenceBindings.cpp
+++ b/Engine/Source/Editor/MovieSceneTools/Private/K2Node_GetSequenceBindings.cpp
@@ -139,7 +139,7 @@ void UDEPRECATED_K2Node_GetSequenceBindings::UpdatePins()
 			}
 
 			FString GuidString = Possessable.GetGuid().ToString();
-			UEdGraphPin* NewPin = CreatePin(EGPD_Output, UEdGraphSchema_K2::PC_Struct, TEXT(""), FMovieSceneObjectBindingID::StaticStruct(), false, false, GuidString);
+			UEdGraphPin* NewPin = CreatePin(EGPD_Output, UEdGraphSchema_K2::PC_Struct, FString(), FMovieSceneObjectBindingID::StaticStruct(), GuidString);
 			NewPin->PinFriendlyName = MovieScene->GetObjectDisplayName(Possessable.GetGuid());
 			NewPin->PersistentGuid = Possessable.GetGuid();
 
diff --git a/Engine/Source/Editor/NiagaraEditor/Private/EdGraphSchema_Niagara.cpp b/Engine/Source/Editor/NiagaraEditor/Private/EdGraphSchema_Niagara.cpp
--- a/Engine/Source/Editor/NiagaraEditor/Private/EdGraphSchema_Niagara.cpp
+++ b/Engine/Source/Editor/NiagaraEditor/Private/EdGraphSchema_Niagara.cpp
@@ -571,7 +571,7 @@ void UEdGraphSchema_Niagara::GetGraphContextActions(FGraphContextMenuBuilder& Co
 					}
 				};
 
-				GenSwizzles(TEXT(""));
+				GenSwizzles(FString());
 
 				for (FString Swiz : Swizzles)
 				{
@@ -773,7 +773,7 @@ const FPinConnectionResponse UEdGraphSchema_Niagara::CanCreateConnection(const U
 	}
 	else
 	{
-		return FPinConnectionResponse(CONNECT_RESPONSE_MAKE, TEXT(""));
+		return FPinConnectionResponse(CONNECT_RESPONSE_MAKE, FString());
 	}
 }
 
@@ -1051,12 +1051,12 @@ FEdGraphPinType UEdGraphSchema_Niagara::TypeDefinitionToPinType(FNiagaraTypeDefi
 {
 	if (TypeDef.GetClass())
 	{
-		return FEdGraphPinType(PinCategoryClass, TEXT(""), const_cast<UClass*>(TypeDef.GetClass()), false, false, false, false, FEdGraphTerminalType());
+		return FEdGraphPinType(PinCategoryClass, FString(), const_cast<UClass*>(TypeDef.GetClass()), EPinContainerType::None, false, FEdGraphTerminalType());
 	}
 	else
 	{
 		//TODO: Are base types better as structs or done like BPS as a special name?
-		return FEdGraphPinType(PinCategoryType, TEXT(""), const_cast<UScriptStruct*>(TypeDef.GetScriptStruct()), false, false, false, false, FEdGraphTerminalType());
+		return FEdGraphPinType(PinCategoryType, FString(), const_cast<UScriptStruct*>(TypeDef.GetScriptStruct()), EPinContainerType::None, false, FEdGraphTerminalType());
 	}
 }
 
@@ -1099,7 +1099,7 @@ void UEdGraphSchema_Niagara::GetBreakLinkToSubMenuActions(class FMenuBuilder& Me
 		UEdGraphPin* Pin = *Links;
 		FString TitleString = Pin->GetOwningNode()->GetNodeTitle(ENodeTitleType::ListView).ToString();
 		FText Title = FText::FromString(TitleString);
-		if (Pin->PinName != TEXT(""))
+		if (!Pin->PinName.IsEmpty())
 		{
 			TitleString = FString::Printf(TEXT("%s (%s)"), *TitleString, *Pin->PinName);
 
diff --git a/Engine/Source/Editor/NiagaraEditor/Private/NiagaraNodeWithDynamicPins.cpp b/Engine/Source/Editor/NiagaraEditor/Private/NiagaraNodeWithDynamicPins.cpp
--- a/Engine/Source/Editor/NiagaraEditor/Private/NiagaraNodeWithDynamicPins.cpp
+++ b/Engine/Source/Editor/NiagaraEditor/Private/NiagaraNodeWithDynamicPins.cpp
@@ -79,7 +79,7 @@ UEdGraphPin* UNiagaraNodeWithDynamicPins::RequestNewTypedPin(EEdGraphPinDirectio
 
 void UNiagaraNodeWithDynamicPins::CreateAddPin(EEdGraphPinDirection Direction)
 {
-	CreatePin(Direction, FEdGraphPinType(UEdGraphSchema_Niagara::PinCategoryMisc, AddPinSubCategory, nullptr, false, false, false, false, FEdGraphTerminalType()), TEXT("Add"));
+	CreatePin(Direction, FEdGraphPinType(UEdGraphSchema_Niagara::PinCategoryMisc, AddPinSubCategory, nullptr, EPinContainerType::None, false, FEdGraphTerminalType()), TEXT("Add"));
 }
 
 bool IsAddPin(const UEdGraphPin* Pin) 
diff --git a/Engine/Source/Editor/Persona/Private/AnimationEditorViewportClient.cpp b/Engine/Source/Editor/Persona/Private/AnimationEditorViewportClient.cpp
--- a/Engine/Source/Editor/Persona/Private/AnimationEditorViewportClient.cpp
+++ b/Engine/Source/Editor/Persona/Private/AnimationEditorViewportClient.cpp
@@ -578,7 +578,7 @@ void FAnimationViewportClient::ShowBoneNames( FCanvas* Canvas, FSceneView* View
 		const FColor BoneColor = FColor::White;
 		if (BoneColor.A != 0)
 		{
-			const FVector BonePos = PreviewMeshComponent->ComponentToWorld.TransformPosition(PreviewMeshComponent->GetComponentSpaceTransforms()[BoneIndex].GetLocation());
+			const FVector BonePos = PreviewMeshComponent->GetComponentTransform().TransformPosition(PreviewMeshComponent->GetComponentSpaceTransforms()[BoneIndex].GetLocation());
 
 			const FPlane proj = View->Project(BonePos);
 			if (proj.W > 0.f)
@@ -1102,7 +1102,7 @@ void FAnimationViewportClient::DrawBonesFromTransforms(TArray<FTransform>& Trans
 			const int32 BoneIndex = MeshComponent->RequiredBones[Index];
 			const int32 ParentIndex = MeshComponent->SkeletalMesh->RefSkeleton.GetParentIndex(BoneIndex);
 
-			WorldTransforms[BoneIndex] = Transforms[BoneIndex] * MeshComponent->ComponentToWorld;
+			WorldTransforms[BoneIndex] = Transforms[BoneIndex] * MeshComponent->GetComponentTransform();
 			BoneColours[BoneIndex] = (ParentIndex >= 0) ? BoneColour : RootBoneColour;
 		}
 
@@ -1129,7 +1129,7 @@ void FAnimationViewportClient::DrawBonesFromCompactPose(const FCompactHeapPose&
 
 			if (ParentIndex == INDEX_NONE)
 			{
-				WorldTransforms[MeshBoneIndex.GetInt()] = Pose[BoneIndex] * MeshComponent->ComponentToWorld;
+				WorldTransforms[MeshBoneIndex.GetInt()] = Pose[BoneIndex] * MeshComponent->GetComponentTransform();
 			}
 			else
 			{
@@ -1224,7 +1224,7 @@ void FAnimationViewportClient::DrawMeshBones(USkeletalMeshComponent * MeshCompon
 			const int32 BoneIndex = MeshComponent->RequiredBones[Index];
 			const int32 ParentIndex = MeshComponent->SkeletalMesh->RefSkeleton.GetParentIndex(BoneIndex);
 
-			WorldTransforms[BoneIndex] = MeshComponent->GetComponentSpaceTransforms()[BoneIndex] * MeshComponent->ComponentToWorld;
+			WorldTransforms[BoneIndex] = MeshComponent->GetComponentSpaceTransforms()[BoneIndex] * MeshComponent->GetComponentTransform();
 			
 			if(SelectedBones.Contains(BoneIndex))
 			{
@@ -1375,7 +1375,7 @@ void FAnimationViewportClient::DrawMeshSubsetBones(const USkeletalMeshComponent*
 					//found a bone we are interested in
 					if(ParentIndex >= 0)
 					{
-						WorldTransforms[ParentIndex] = MeshComponent->GetComponentSpaceTransforms()[ParentIndex]*MeshComponent->ComponentToWorld;
+						WorldTransforms[ParentIndex] = MeshComponent->GetComponentSpaceTransforms()[ParentIndex]*MeshComponent->GetComponentTransform();
 					}
 					BoneColours[BoneIndex] = LinearSelectionColor;
 					bDrawBone = true;
@@ -1393,7 +1393,7 @@ void FAnimationViewportClient::DrawMeshSubsetBones(const USkeletalMeshComponent*
 			{
 				//add to the list
 				RequiredBones.AddUnique(BoneIndex);
-				WorldTransforms[BoneIndex] = MeshComponent->GetComponentSpaceTransforms()[BoneIndex] * MeshComponent->ComponentToWorld;
+				WorldTransforms[BoneIndex] = MeshComponent->GetComponentSpaceTransforms()[BoneIndex] * MeshComponent->GetComponentTransform();
 			}
 		}
 
@@ -1420,7 +1420,7 @@ void FAnimationViewportClient::DrawSockets(const UDebugSkelMeshComponent* InPrev
 			FVector Start, End;
 			if (ParentIndex >=0)
 			{
-				FTransform WorldTransformParent = InPreviewMeshComponent->GetComponentSpaceTransforms()[ParentIndex] * InPreviewMeshComponent->ComponentToWorld;
+				FTransform WorldTransformParent = InPreviewMeshComponent->GetComponentSpaceTransforms()[ParentIndex] * InPreviewMeshComponent->GetComponentTransform();
 				Start = WorldTransformParent.GetLocation();
 				End = WorldTransformSocket.GetLocation();
 			}
@@ -1477,8 +1477,7 @@ FSphere FAnimationViewportClient::GetCameraTarget()
 	UDebugSkelMeshComponent* PreviewMeshComponent = GetAnimPreviewScene()->GetPreviewMeshComponent();
 	if( PreviewMeshComponent != nullptr )
 	{
-		FTransform ComponentToWorld = PreviewMeshComponent->ComponentToWorld;
-		PreviewMeshComponent->CalcBounds(ComponentToWorld);
+		PreviewMeshComponent->CalcBounds(PreviewMeshComponent->GetComponentTransform());
 
 		// give the editor mode a chance to give us a camera target
 		FSphere Target;
diff --git a/Engine/Source/Editor/Persona/Private/SAnimationScrubPanel.cpp b/Engine/Source/Editor/Persona/Private/SAnimationScrubPanel.cpp
--- a/Engine/Source/Editor/Persona/Private/SAnimationScrubPanel.cpp
+++ b/Engine/Source/Editor/Persona/Private/SAnimationScrubPanel.cpp
@@ -575,8 +575,8 @@ void SAnimationScrubPanel::OnReZeroAnimSequence(int32 FrameIndex)
 				FVector ApplyTranslation = -1.f * FrameTransform;
 
 				// Convert into world space
-				FVector WorldApplyTranslation = PreviewSkelComp->ComponentToWorld.TransformVector(ApplyTranslation);
-				ApplyTranslation = PreviewSkelComp->ComponentToWorld.InverseTransformVector(WorldApplyTranslation);
+				FVector WorldApplyTranslation = PreviewSkelComp->GetComponentTransform().TransformVector(ApplyTranslation);
+				ApplyTranslation = PreviewSkelComp->GetComponentTransform().InverseTransformVector(WorldApplyTranslation);
 
 				for(int32 i=0; i<RawTrack.PosKeys.Num(); i++)
 				{
diff --git a/Engine/Source/Editor/PhAT/Private/PhATPreviewViewportClient.cpp b/Engine/Source/Editor/PhAT/Private/PhATPreviewViewportClient.cpp
--- a/Engine/Source/Editor/PhAT/Private/PhATPreviewViewportClient.cpp
+++ b/Engine/Source/Editor/PhAT/Private/PhATPreviewViewportClient.cpp
@@ -61,7 +61,7 @@ FPhATEdPreviewViewportClient::FPhATEdPreviewViewportClient(TWeakPtr<FPhAT> InPhA
 	EngineShowFlags.SetCompositeEditorPrimitives(true);
 
 	// Get actors asset collision bounding box, and move actor so its not intersection the floor plane at Z = 0.
-	FBox CollBox = SharedData->PhysicsAsset->CalcAABB(SharedData->EditorSkelComp, SharedData->EditorSkelComp->ComponentToWorld);	
+	FBox CollBox = SharedData->PhysicsAsset->CalcAABB(SharedData->EditorSkelComp, SharedData->EditorSkelComp->GetComponentTransform());	
 	FVector SkelCompLocation = FVector(0, 0, -CollBox.Min.Z + SharedData->EditorSimOptions->FloorGap);
 
 	SharedData->EditorSkelComp->SetAbsolute(true, true, true);
@@ -69,11 +69,11 @@ FPhATEdPreviewViewportClient::FPhATEdPreviewViewportClient(TWeakPtr<FPhAT> InPhA
 	SharedData->ResetTM = SharedData->EditorSkelComp->GetComponentToWorld();
 
 	// Get new bounding box and set view based on that.
-	CollBox = SharedData->PhysicsAsset->CalcAABB(SharedData->EditorSkelComp, SharedData->EditorSkelComp->ComponentToWorld);	
+	CollBox = SharedData->PhysicsAsset->CalcAABB(SharedData->EditorSkelComp, SharedData->EditorSkelComp->GetComponentTransform());	
 	FVector CollBoxExtent = CollBox.GetExtent();
 
 	// Take into account internal mesh translation/rotation/scaling etc.
-	FTransform LocalToWorld = SharedData->EditorSkelComp->ComponentToWorld;
+	FTransform LocalToWorld = SharedData->EditorSkelComp->GetComponentTransform();
 	FSphere WorldSphere = SharedData->EditorSkelMesh->GetImportedBounds().GetSphere().TransformBy(LocalToWorld);
 
 	CollBoxExtent = CollBox.GetExtent();
@@ -178,7 +178,7 @@ void FPhATEdPreviewViewportClient::DrawCanvas( FViewport& InViewport, FSceneView
 		// Iterate over each graphics bone.
 		for (int32 i = 0; i < SharedData->EditorSkelComp->GetNumComponentSpaceTransforms(); ++i)
 		{
-			FVector BonePos = SharedData->EditorSkelComp->ComponentToWorld.TransformPosition(SharedData->EditorSkelComp->GetComponentSpaceTransforms()[i].GetLocation());
+			FVector BonePos = SharedData->EditorSkelComp->GetComponentTransform().TransformPosition(SharedData->EditorSkelComp->GetComponentSpaceTransforms()[i].GetLocation());
 
 			FPlane proj = View.Project(BonePos);
 			if (proj.W > 0.f) // This avoids drawing bone names that are behind us.
diff --git a/Engine/Source/Editor/PhAT/Private/PhATRender.cpp b/Engine/Source/Editor/PhAT/Private/PhATRender.cpp
--- a/Engine/Source/Editor/PhAT/Private/PhATRender.cpp
+++ b/Engine/Source/Editor/PhAT/Private/PhATRender.cpp
@@ -288,13 +288,13 @@ void UPhATEdSkeletalMeshComponent::DrawHierarchy(FPrimitiveDrawInterface* PDI, b
 		FVector ParentPos, ChildPos;
 		if (bAnimSkel)
 		{
-			ParentPos = ComponentToWorld.TransformPosition(AnimationSpaceBases[ParentIndex].GetLocation());
-			ChildPos = ComponentToWorld.TransformPosition(AnimationSpaceBases[i].GetLocation());
+			ParentPos = GetComponentTransform().TransformPosition(AnimationSpaceBases[ParentIndex].GetLocation());
+			ChildPos = GetComponentTransform().TransformPosition(AnimationSpaceBases[i].GetLocation());
 		}
 		else
 		{
-			ParentPos = ComponentToWorld.TransformPosition(GetComponentSpaceTransforms()[ParentIndex].GetLocation());
-			ChildPos = ComponentToWorld.TransformPosition(GetComponentSpaceTransforms()[i].GetLocation());
+			ParentPos = GetComponentTransform().TransformPosition(GetComponentSpaceTransforms()[ParentIndex].GetLocation());
+			ChildPos = GetComponentTransform().TransformPosition(GetComponentSpaceTransforms()[i].GetLocation());
 		}
 
 		FColor DrawColor = bAnimSkel ? AnimSkelDrawColor : HierarchyDrawColor;
diff --git a/Engine/Source/Editor/ReferenceViewer/Private/EdGraphNode_Reference.cpp b/Engine/Source/Editor/ReferenceViewer/Private/EdGraphNode_Reference.cpp
--- a/Engine/Source/Editor/ReferenceViewer/Private/EdGraphNode_Reference.cpp
+++ b/Engine/Source/Editor/ReferenceViewer/Private/EdGraphNode_Reference.cpp
@@ -165,8 +165,8 @@ FText UEdGraphNode_Reference::GetTooltipText() const
 
 void UEdGraphNode_Reference::AllocateDefaultPins()
 {
-	ReferencerPin = CreatePin( EEdGraphPinDirection::EGPD_Input, TEXT(""), TEXT(""), NULL, false, false, TEXT("") );
-	DependencyPin = CreatePin( EEdGraphPinDirection::EGPD_Output, TEXT(""), TEXT(""), NULL, false, false, TEXT("") );
+	ReferencerPin = CreatePin( EEdGraphPinDirection::EGPD_Input, FString(), FString(), nullptr, FString() );
+	DependencyPin = CreatePin( EEdGraphPinDirection::EGPD_Output, FString(), FString(), nullptr, FString() );
 
 	ReferencerPin->bHidden = true;
 	DependencyPin->bHidden = true;
diff --git a/Engine/Source/Editor/ReferenceViewer/Private/ReferenceViewerSchema.cpp b/Engine/Source/Editor/ReferenceViewer/Private/ReferenceViewerSchema.cpp
--- a/Engine/Source/Editor/ReferenceViewer/Private/ReferenceViewerSchema.cpp
+++ b/Engine/Source/Editor/ReferenceViewer/Private/ReferenceViewerSchema.cpp
@@ -105,16 +105,16 @@ void UReferenceViewerSchema::BreakSinglePinLink(UEdGraphPin* SourcePin, UEdGraph
 	// Don't allow breaking any links
 }
 
-FPinConnectionResponse UReferenceViewerSchema::MovePinLinks(UEdGraphPin& MoveFromPin, UEdGraphPin& MoveToPin, bool bIsItermeadiateMove) const
+FPinConnectionResponse UReferenceViewerSchema::MovePinLinks(UEdGraphPin& MoveFromPin, UEdGraphPin& MoveToPin, bool bIsItermediateMove) const
 {
 	// Don't allow moving any links
-	return FPinConnectionResponse(CONNECT_RESPONSE_DISALLOW, TEXT(""));
+	return FPinConnectionResponse(CONNECT_RESPONSE_DISALLOW, FString());
 }
 
-FPinConnectionResponse UReferenceViewerSchema::CopyPinLinks(UEdGraphPin& CopyFromPin, UEdGraphPin& CopyToPin, bool bIsItermeadiateCopy) const
+FPinConnectionResponse UReferenceViewerSchema::CopyPinLinks(UEdGraphPin& CopyFromPin, UEdGraphPin& CopyToPin, bool bIsItermediateCopy) const
 {
 	// Don't allow copying any links
-	return FPinConnectionResponse(CONNECT_RESPONSE_DISALLOW, TEXT(""));
+	return FPinConnectionResponse(CONNECT_RESPONSE_DISALLOW, FString());
 }
 
 FConnectionDrawingPolicy* UReferenceViewerSchema::CreateConnectionDrawingPolicy(int32 InBackLayerID, int32 InFrontLayerID, float InZoomFactor, const FSlateRect& InClippingRect, class FSlateWindowElementList& InDrawElements, class UEdGraph* InGraphObj) const
diff --git a/Engine/Source/Editor/SequenceRecorder/Private/AnimationRecorder.cpp b/Engine/Source/Editor/SequenceRecorder/Private/AnimationRecorder.cpp
--- a/Engine/Source/Editor/SequenceRecorder/Private/AnimationRecorder.cpp
+++ b/Engine/Source/Editor/SequenceRecorder/Private/AnimationRecorder.cpp
@@ -212,7 +212,7 @@ void FAnimationRecorder::StartRecord(USkeletalMeshComponent* Component, UAnimSeq
 
 	GetBoneTransforms(Component, PreviousSpacesBases);
 	PreviousAnimCurves = Component->GetAnimationCurves();
-	PreviousComponentToWorld = Component->ComponentToWorld;
+	PreviousComponentToWorld = Component->GetComponentTransform();
 
 	LastFrame = 0;
 	AnimationObject->SequenceLength = 0.f;
@@ -484,7 +484,7 @@ void FAnimationRecorder::UpdateRecord(USkeletalMeshComponent* Component, float D
 			}
 
 			FTransform BlendedComponentToWorld;
-			BlendedComponentToWorld.Blend(PreviousComponentToWorld, Component->ComponentToWorld, BlendAlpha);
+			BlendedComponentToWorld.Blend(PreviousComponentToWorld, Component->GetComponentTransform(), BlendAlpha);
 
 			FBlendedHeapCurve BlendedCurve;
 			if (AnimCurves.Elements.Num() > 0 && PreviousAnimCurves.Elements.Num() == AnimCurves.Elements.Num())
@@ -506,7 +506,7 @@ void FAnimationRecorder::UpdateRecord(USkeletalMeshComponent* Component, float D
 	//save to current transform
 	PreviousSpacesBases = SpaceBases;
 	PreviousAnimCurves = Component->GetAnimationCurves();
-	PreviousComponentToWorld = Component->ComponentToWorld;
+	PreviousComponentToWorld = Component->GetComponentTransform();
 
 	// if we passed MaxFrame, just stop it
 	if (MaxFrame != UnBoundedFrameCount && FramesRecorded >= MaxFrame)
diff --git a/Engine/Source/Editor/SourceControlWindows/Private/SSourceControlHistory.cpp b/Engine/Source/Editor/SourceControlWindows/Private/SSourceControlHistory.cpp
--- a/Engine/Source/Editor/SourceControlWindows/Private/SSourceControlHistory.cpp
+++ b/Engine/Source/Editor/SourceControlWindows/Private/SSourceControlHistory.cpp
@@ -206,11 +206,8 @@ static UObject* GetAssetRevisionObject(TSharedPtr<FHistoryTreeItem> HistoryTreeI
 				FString TempPackageName;
 				if (FileRevision.IsValid() && FileRevision->Get(TempPackageName)) // grab the path to a temporary package (where the revision item will be stored)
 				{
-					// forcibly disable compile on load in case we are loading old blueprints that might try to update/compile
-					TGuardValue<bool> DisableCompileOnLoad(GForceDisableBlueprintCompileOnLoad, true);
-
 					// try and load the temporary package
-					AssetPackage = LoadPackage(NULL, *TempPackageName, LOAD_ForDiff);
+					AssetPackage = LoadPackage(NULL, *TempPackageName, LOAD_DisableCompileOnLoad);
 				}
 			} // if FileSourceControlState.IsValid()
 		}
diff --git a/Engine/Source/Editor/StaticMeshEditor/Private/SStaticMeshEditorViewport.cpp b/Engine/Source/Editor/StaticMeshEditor/Private/SStaticMeshEditorViewport.cpp
--- a/Engine/Source/Editor/StaticMeshEditor/Private/SStaticMeshEditorViewport.cpp
+++ b/Engine/Source/Editor/StaticMeshEditor/Private/SStaticMeshEditorViewport.cpp
@@ -223,7 +223,7 @@ void SStaticMeshEditorViewport::UpdatePreviewSocketMeshes()
 					SocketPreviewMeshComponent = NewObject<UStaticMeshComponent>();
 					PreviewScene->AddComponent(SocketPreviewMeshComponent, FTransform::Identity);
 					SocketPreviewMeshComponents.Add(SocketPreviewMeshComponent);
-					SocketPreviewMeshComponent->SnapTo(PreviewMeshComponent, Socket->SocketName);
+					SocketPreviewMeshComponent->AttachToComponent(PreviewMeshComponent, FAttachmentTransformRules::SnapToTargetNotIncludingScale, Socket->SocketName);
 				}
 				else
 				{
@@ -232,7 +232,7 @@ void SStaticMeshEditorViewport::UpdatePreviewSocketMeshes()
 					// In case of a socket rename, ensure our preview component is still snapping to the proper socket
 					if (!SocketPreviewMeshComponent->GetAttachSocketName().IsEqual(Socket->SocketName))
 					{
-						SocketPreviewMeshComponent->SnapTo(PreviewMeshComponent, Socket->SocketName);
+						SocketPreviewMeshComponent->AttachToComponent(PreviewMeshComponent, FAttachmentTransformRules::SnapToTargetNotIncludingScale, Socket->SocketName);
 					}
 
 					// Force component to world update to take into account the new socket position.
@@ -295,7 +295,7 @@ void SStaticMeshEditorViewport::UpdatePreviewMesh(UStaticMesh* InStaticMesh, boo
 		{
 			SocketPreviewMeshComponent = NewObject<UStaticMeshComponent>();
 			SocketPreviewMeshComponent->SetStaticMesh(Socket->PreviewStaticMesh);
-			SocketPreviewMeshComponent->SnapTo(PreviewMeshComponent, Socket->SocketName);
+			SocketPreviewMeshComponent->AttachToComponent(PreviewMeshComponent, FAttachmentTransformRules::SnapToTargetNotIncludingScale, Socket->SocketName);
 			SocketPreviewMeshComponents.Add(SocketPreviewMeshComponent);
 			PreviewScene->AddComponent(SocketPreviewMeshComponent, FTransform::Identity);
 		}
diff --git a/Engine/Source/Editor/StaticMeshEditor/Private/StaticMeshEditorViewportClient.cpp b/Engine/Source/Editor/StaticMeshEditor/Private/StaticMeshEditorViewportClient.cpp
--- a/Engine/Source/Editor/StaticMeshEditor/Private/StaticMeshEditorViewportClient.cpp
+++ b/Engine/Source/Editor/StaticMeshEditor/Private/StaticMeshEditorViewportClient.cpp
@@ -544,8 +544,8 @@ void FStaticMeshEditorViewportClient::Draw(const FSceneView* View,FPrimitiveDraw
 			EdgeVertices[ 1 ] = SelectedEdgeVertices[VertexIndex + 1];
 
 			PDI->DrawLine(
-				StaticMeshComponent->ComponentToWorld.TransformPosition( EdgeVertices[ 0 ] ),
-				StaticMeshComponent->ComponentToWorld.TransformPosition( EdgeVertices[ 1 ] ),
+				StaticMeshComponent->GetComponentTransform().TransformPosition( EdgeVertices[ 0 ] ),
+				StaticMeshComponent->GetComponentTransform().TransformPosition( EdgeVertices[ 1 ] ),
 				FColor( 255, 255, 0 ),
 				SDPG_World );
 		}
@@ -558,12 +558,12 @@ void FStaticMeshEditorViewportClient::Draw(const FSceneView* View,FPrimitiveDraw
 		FIndexArrayView Indices = LODModel.IndexBuffer.GetArrayView();
 		uint32 NumIndices = Indices.Num();
 
-		FMatrix LocalToWorldInverseTranspose = StaticMeshComponent->ComponentToWorld.ToMatrixWithScale().InverseFast().GetTransposed();
+		FMatrix LocalToWorldInverseTranspose = StaticMeshComponent->GetComponentTransform().ToMatrixWithScale().InverseFast().GetTransposed();
 		for (uint32 i = 0; i < NumIndices; i++)
 		{
 			const FVector& VertexPos = LODModel.PositionVertexBuffer.VertexPosition( Indices[i] );
 
-			const FVector WorldPos = StaticMeshComponent->ComponentToWorld.TransformPosition( VertexPos );
+			const FVector WorldPos = StaticMeshComponent->GetComponentTransform().TransformPosition( VertexPos );
 			const FVector& Normal = LODModel.VertexBuffer.VertexTangentZ( Indices[i] ); 
 			const FVector& Binormal = LODModel.VertexBuffer.VertexTangentY( Indices[i] ); 
 			const FVector& Tangent = LODModel.VertexBuffer.VertexTangentX( Indices[i] ); 
@@ -599,7 +599,7 @@ void FStaticMeshEditorViewportClient::Draw(const FSceneView* View,FPrimitiveDraw
 
 	if( bShowPivot )
 	{
-		FUnrealEdUtils::DrawWidget(View, PDI, StaticMeshComponent->ComponentToWorld.ToMatrixWithScale(), 0, 0, EAxisList::All, EWidgetMovementMode::WMM_Translate, false);
+		FUnrealEdUtils::DrawWidget(View, PDI, StaticMeshComponent->GetComponentTransform().ToMatrixWithScale(), 0, 0, EAxisList::All, EWidgetMovementMode::WMM_Translate, false);
 	}
 
 	if( bDrawAdditionalData )
@@ -620,7 +620,7 @@ void FStaticMeshEditorViewportClient::Draw(const FSceneView* View,FPrimitiveDraw
 		if (StaticMesh->NavCollision && StaticMesh->NavCollision->bIsDynamicObstacle)
 		{
 			// Draw the static mesh's body setup (simple collision)
-			FTransform GeomTransform(StaticMeshComponent->ComponentToWorld);
+			FTransform GeomTransform(StaticMeshComponent->GetComponentTransform());
 			FColor NavCollisionColor = FColor(118, 84, 255, 255);
 			StaticMesh->NavCollision->DrawSimpleGeom(PDI, GeomTransform, FColorList::LimeGreen);
 		}
@@ -1028,7 +1028,7 @@ void FStaticMeshEditorViewportClient::ProcessClick(class FSceneView& InView, cla
 						const FVector TriangleNormal = (CA ^ BA).GetSafeNormal();
 
 						// Transform the view position from world to component space
-						const FVector ComponentSpaceViewOrigin = StaticMeshComponent->ComponentToWorld.InverseTransformPosition( View->ViewMatrices.GetViewOrigin());
+						const FVector ComponentSpaceViewOrigin = StaticMeshComponent->GetComponentTransform().InverseTransformPosition( View->ViewMatrices.GetViewOrigin());
 								
 						// Determine which side of the triangle's plane that the view position lies on.
 						bIsBackFacing = (FVector::PointPlaneDist( ComponentSpaceViewOrigin,  A, TriangleNormal)  < 0.0f);
@@ -1062,8 +1062,8 @@ void FStaticMeshEditorViewportClient::ProcessClick(class FSceneView& InView, cla
 						}
 						else
 						{
-							FVector WorldSpaceEdgeStart( StaticMeshComponent->ComponentToWorld.TransformPosition( EdgeVertices[ 0 ] ) );
-							FVector WorldSpaceEdgeEnd( StaticMeshComponent->ComponentToWorld.TransformPosition( EdgeVertices[ 1 ] ) );
+							FVector WorldSpaceEdgeStart( StaticMeshComponent->GetComponentTransform().TransformPosition( EdgeVertices[ 0 ] ) );
+							FVector WorldSpaceEdgeEnd( StaticMeshComponent->GetComponentTransform().TransformPosition( EdgeVertices[ 1 ] ) );
 
 							// Determine the mesh edge that's closest to the ray cast through the eye towards the click location
 							FVector ClosestPointToEdgeOnClickLine;
diff --git a/Engine/Source/Editor/UMGEditor/Private/Nodes/K2Node_CreateWidget.cpp b/Engine/Source/Editor/UMGEditor/Private/Nodes/K2Node_CreateWidget.cpp
--- a/Engine/Source/Editor/UMGEditor/Private/Nodes/K2Node_CreateWidget.cpp
+++ b/Engine/Source/Editor/UMGEditor/Private/Nodes/K2Node_CreateWidget.cpp
@@ -37,7 +37,7 @@ void UK2Node_CreateWidget::AllocateDefaultPins()
 	const UEdGraphSchema_K2* K2Schema = GetDefault<UEdGraphSchema_K2>();
 
 	// OwningPlayer pin
-	UEdGraphPin* OwningPlayerPin = CreatePin(EGPD_Input, K2Schema->PC_Object, TEXT(""), APlayerController::StaticClass(), false, false, FK2Node_CreateWidgetHelper::OwningPlayerPinName);
+	UEdGraphPin* OwningPlayerPin = CreatePin(EGPD_Input, K2Schema->PC_Object, FString(), APlayerController::StaticClass(), FK2Node_CreateWidgetHelper::OwningPlayerPinName);
 	SetPinToolTip(*OwningPlayerPin, LOCTEXT("OwningPlayerPinDescription", "The player that 'owns' the widget."));
 }
 
diff --git a/Engine/Source/Editor/UMGEditor/Private/WidgetBlueprint.cpp b/Engine/Source/Editor/UMGEditor/Private/WidgetBlueprint.cpp
--- a/Engine/Source/Editor/UMGEditor/Private/WidgetBlueprint.cpp
+++ b/Engine/Source/Editor/UMGEditor/Private/WidgetBlueprint.cpp
@@ -740,9 +740,59 @@ bool UWidgetBlueprint::IsWidgetFreeFromCircularReferences(UUserWidget* UserWidge
 	return true;
 }
 
+TArray<UWidget*> UWidgetBlueprint::GetAllSourceWidgets()
+{
+	TArray<UWidget*> Ret;
+	ForEachSourceWidgetImpl( [&Ret](UWidget* Inner) { Ret.Push(Inner); } );
+	return Ret;
+}
+
+TArray<const UWidget*> UWidgetBlueprint::GetAllSourceWidgets() const
+{
+	TArray<const UWidget*> Ret;
+	ForEachSourceWidgetImpl( [&Ret](UWidget* Inner) { Ret.Push(Inner); } );
+	return Ret;
+}
+
+void UWidgetBlueprint::ForEachSourceWidget(TFunctionRef<void(UWidget*)> Fn)
+{
+	ForEachSourceWidgetImpl(Fn);
+}
+
+void UWidgetBlueprint::ForEachSourceWidget(TFunctionRef<void(const UWidget*)> Fn) const
+{
+	ForEachSourceWidgetImpl(Fn);
+}
+
 UPackage* UWidgetBlueprint::GetWidgetTemplatePackage() const
 {
 	return GetOutermost();
 }
 
+void UWidgetBlueprint::ForEachSourceWidgetImpl (TFunctionRef<void(UWidget*)> Fn) const
+{
+	// This exists in order to facilitate working with collections of UWidgets wo/ 
+	// relying on user implemented UWidget virtual functions. During blueprint compilation
+	// it is bad practice to call those virtual functions until the class is fully formed
+	// and reinstancing has finished. For instance, GetDefaultObject() calls in those user
+	// functions may create a CDO before the class has been linked, or even before
+	// all member variables have been generated:
+	UWidgetTree* WidgetTreeForCapture = WidgetTree;
+	ForEachObjectWithOuter(
+		WidgetTree, 
+		[Fn, WidgetTreeForCapture](UObject* Inner)
+		{
+			if(UWidget* AsWidget = Cast<UWidget>(Inner))
+			{
+				// Widgets owned by another UWidgetBlueprint aren't really 'source' widgets, E.g. widgets
+				// created by the user *in this blueprint*
+				if(AsWidget->GetTypedOuter<UWidgetTree>() == WidgetTreeForCapture)
+				{
+					Fn(AsWidget);
+				}
+			}
+		}
+	);
+}
+
 #undef LOCTEXT_NAMESPACE 
diff --git a/Engine/Source/Editor/UMGEditor/Private/WidgetBlueprintCompiler.cpp b/Engine/Source/Editor/UMGEditor/Private/WidgetBlueprintCompiler.cpp
--- a/Engine/Source/Editor/UMGEditor/Private/WidgetBlueprintCompiler.cpp
+++ b/Engine/Source/Editor/UMGEditor/Private/WidgetBlueprintCompiler.cpp
@@ -223,8 +223,7 @@ void FWidgetBlueprintCompiler::CreateClassVariablesFromBlueprint()
 	ValidateWidgetNames();
 
 	// Build the set of variables based on the variable widgets in the widget tree.
-	TArray<UWidget*> Widgets;
-	WidgetBP->WidgetTree->GetAllWidgets(Widgets);
+	TArray<UWidget*> Widgets = WidgetBP->GetAllSourceWidgets();
 
 	// Sort the widgets alphabetically
 	Widgets.Sort( []( const UWidget& Lhs, const UWidget& Rhs ) { return Rhs.GetFName() < Lhs.GetFName(); } );
@@ -267,7 +266,7 @@ void FWidgetBlueprintCompiler::CreateClassVariablesFromBlueprint()
 			continue;
 		}
 
-		FEdGraphPinType WidgetPinType(Schema->PC_Object, TEXT(""), WidgetClass, false, false, false, false, FEdGraphTerminalType());
+		FEdGraphPinType WidgetPinType(Schema->PC_Object, FString(), WidgetClass, EPinContainerType::None, false, FEdGraphTerminalType());
 		
 		// Always name the variable according to the underlying FName of the widget object
 		UProperty* WidgetProperty = CreateVariable(Widget->GetFName(), WidgetPinType);
@@ -275,12 +274,14 @@ void FWidgetBlueprintCompiler::CreateClassVariablesFromBlueprint()
 		{
 			const FString VariableName = Widget->IsGeneratedName() ? Widget->GetName() : Widget->GetLabelText().ToString();
 			WidgetProperty->SetMetaData(TEXT("DisplayName"), *VariableName);
-			WidgetProperty->SetMetaData(TEXT("Category"), *WidgetBP->GetName());
 			
 			// Only show variables if they're explicitly marked as variables.
 			if ( Widget->bIsVariable )
 			{
 				WidgetProperty->SetPropertyFlags(CPF_BlueprintVisible);
+
+				// Only include Category metadata for variables (i.e. a visible/editable property); otherwise, UHT will raise a warning if this Blueprint is nativized.
+				WidgetProperty->SetMetaData(TEXT("Category"), *WidgetBP->GetName());
 			}
 
 			WidgetProperty->SetPropertyFlags(CPF_Instanced);
@@ -293,7 +294,7 @@ void FWidgetBlueprintCompiler::CreateClassVariablesFromBlueprint()
 	// Add movie scenes variables here
 	for(UWidgetAnimation* Animation : WidgetBP->Animations)
 	{
-		FEdGraphPinType WidgetPinType(Schema->PC_Object, TEXT(""), Animation->GetClass(), false, true, false, false, FEdGraphTerminalType());
+		FEdGraphPinType WidgetPinType(Schema->PC_Object, FString(), Animation->GetClass(), EPinContainerType::None, true, FEdGraphTerminalType());
 		UProperty* AnimationProperty = CreateVariable(Animation->GetFName(), WidgetPinType);
 
 		if ( AnimationProperty != nullptr )
@@ -439,7 +440,7 @@ void FWidgetBlueprintCompiler::FinishCompilingClass(UClass* Class)
 
 		// Add all the names of the named slot widgets to the slot names structure.
 		BPGClass->NamedSlots.Reset();
-		BPGClass->WidgetTree->ForEachWidget([&] (UWidget* Widget) {
+		WidgetBP->ForEachSourceWidget([&] (UWidget* Widget) {
 			if ( Widget && Widget->IsA<UNamedSlot>() )
 			{
 				BPGClass->NamedSlots.Add(Widget->GetFName());
@@ -566,7 +567,7 @@ void FWidgetBlueprintCompiler::VerifyEventReplysAreNotEmpty(FKismetFunctionConte
 	Context.SourceGraph->GetNodesOfClass<UK2Node_FunctionResult>(FunctionResults);
 
 	UScriptStruct* EventReplyStruct = FEventReply::StaticStruct();
-	FEdGraphPinType EventReplyPinType(Schema->PC_Struct, TEXT(""), EventReplyStruct, /*bIsArray =*/false, /*bIsReference =*/false, /*bIsSet =*/false, /*bIsMap =*/ false, /*InValueTerminalType =*/FEdGraphTerminalType());
+	FEdGraphPinType EventReplyPinType(Schema->PC_Struct, FString(), EventReplyStruct, EPinContainerType::None, /*bIsReference =*/false, /*InValueTerminalType =*/FEdGraphTerminalType());
 
 	for ( UK2Node_FunctionResult* FunctionResult : FunctionResults )
 	{
diff --git a/Engine/Source/Editor/UMGEditor/Public/WidgetBlueprint.h b/Engine/Source/Editor/UMGEditor/Public/WidgetBlueprint.h
--- a/Engine/Source/Editor/UMGEditor/Public/WidgetBlueprint.h
+++ b/Engine/Source/Editor/UMGEditor/Public/WidgetBlueprint.h
@@ -17,6 +17,7 @@ class FCompilerResultsLog;
 class UEdGraph;
 class UMovieScene;
 class UUserWidget;
+class UWidget;
 class UWidgetAnimation;
 class FKismetCompilerContext;
 
@@ -242,7 +243,22 @@ public:
 	/** Returns true if the supplied user widget will not create a circular reference when added to this blueprint */
 	bool IsWidgetFreeFromCircularReferences(UUserWidget* UserWidget) const;
 
+	/** 
+	 * Returns collection of widgets that represent the 'source' (user edited) widgets for this 
+	 * blueprint - avoids calling virtual functions on instances and is therefore safe to use 
+	 * throughout compilation.
+	 */
+	TArray<UWidget*> GetAllSourceWidgets();
+	TArray<const UWidget*> GetAllSourceWidgets() const;
+
+	/** Identical to GetAllSourceWidgets, but as an algorithm */
+	void ForEachSourceWidget(TFunctionRef<void(UWidget*)> Fn);
+	void ForEachSourceWidget(TFunctionRef<void(const UWidget*)> Fn) const;
+
 	static bool ValidateGeneratedClass(const UClass* InClass);
 	
 	static TUniquePtr<FKismetCompilerContext> GetCompilerForWidgetBP(UWidgetBlueprint* BP, FCompilerResultsLog& InMessageLog, const FKismetCompilerOptions& InCompileOptions);
+
+private:
+	void ForEachSourceWidgetImpl(TFunctionRef<void(UWidget*)> Fn) const;
 };
diff --git a/Engine/Source/Editor/UnrealEd/Classes/Animation/EditorSkeletonNotifyObj.h b/Engine/Source/Editor/UnrealEd/Classes/Animation/EditorSkeletonNotifyObj.h
--- a/Engine/Source/Editor/UnrealEd/Classes/Animation/EditorSkeletonNotifyObj.h
+++ b/Engine/Source/Editor/UnrealEd/Classes/Animation/EditorSkeletonNotifyObj.h
@@ -14,19 +14,6 @@ class UAnimSequenceBase;
 
 DECLARE_DELEGATE_TwoParams( FOnAnimObjectChange, class UObject*, bool)
 
-USTRUCT()
-struct FSkeletonNotifyDependentAnimations
-{
-	GENERATED_USTRUCT_BODY()
-
-	UPROPERTY(EditAnywhere, BlueprintReadWrite, Category=SkeletonNotifies)
-	TArray<UAnimSequenceBase*> Animations;
-
-	FSkeletonNotifyDependentAnimations()
-	{
-	}
-};
-
 UCLASS(MinimalAPI)
 class UEditorSkeletonNotifyObj : public UObject
 {
diff --git a/Engine/Source/Editor/UnrealEd/Classes/Commandlets/ResavePackagesCommandlet.h b/Engine/Source/Editor/UnrealEd/Classes/Commandlets/ResavePackagesCommandlet.h
--- a/Engine/Source/Editor/UnrealEd/Classes/Commandlets/ResavePackagesCommandlet.h
+++ b/Engine/Source/Editor/UnrealEd/Classes/Commandlets/ResavePackagesCommandlet.h
@@ -81,18 +81,21 @@ protected:
 	/** Ignore package version changelist **/
 	bool bIgnoreChangelist;
 
-	// Running count of packages that got modified and will need to be resaved
+	/** Running count of packages that got modified and will need to be resaved */
 	int32 PackagesRequiringResave;
 
 	/** Only collect garbage after N packages */
 	int32 GarbageCollectionFrequency;
 
-	// List of files to submit
+	/** List of files to submit */
 	TArray<FString> FilesToSubmit;
 
-	// The list of switches that were passed on the commandline
+	/** The list of switches that were passed on the commandline */
 	TArray<FString> Switches;
 
+	/** List of redirector packages that should be fixed up at the end */
+	TArray<FString> RedirectorsToFixup;
+
 	/**
 	 * Evaluates the command-line to determine which maps to check.  By default all maps are checked
 	 * Provides child classes with a chance to initialize any variables, parse the command line, etc.
@@ -104,12 +107,15 @@ protected:
 	 */
 	virtual int32 InitializeResaveParameters( const TArray<FString>& Tokens, TArray<FString>& MapPathNames );
 
-	// Loads and saves a single package
+	/** Loads and saves a single package */
 	virtual void LoadAndSaveOnePackage(const FString& Filename);
 
-	// Checks to see if a package should be skipped
+	/** Checks to see if a package should be skipped */
 	virtual bool ShouldSkipPackage(const FString& Filename);
 
+	/** Deletes a single package */
+	virtual void DeleteOnePackage(const FString& Filename);
+
 	/**
 	 * Allow the commandlet to perform any operations on the export/import table of the package before all objects in the package are loaded.
 	 *
