Fixing enum comments to correct (Doxygen standard) order. Affects documentation.

diff --git a/Engine/Source/Editor/KismetCompiler/Public/BlueprintCompiledStatement.h b/Engine/Source/Editor/KismetCompiler/Public/BlueprintCompiledStatement.h
--- a/Engine/Source/Editor/KismetCompiler/Public/BlueprintCompiledStatement.h
+++ b/Engine/Source/Editor/KismetCompiler/Public/BlueprintCompiledStatement.h
@@ -10,48 +10,88 @@
 enum EKismetCompiledStatementType
 {
 	KCST_Nop = 0,
-	KCST_CallFunction = 1,		// [wiring =] TargetObject->FunctionToCall(wiring)
-	KCST_Assignment = 2,		// TargetObject->TargetProperty = [wiring]
-	KCST_CompileError = 3,		// One of the other types with a compilation error during statement generation
-	KCST_UnconditionalGoto = 4,	// goto TargetLabel
-	KCST_PushState = 5,         // FlowStack.Push(TargetLabel)
-	KCST_GotoIfNot = 6,			// [if (!TargetObject->TargetProperty)] goto TargetLabel
-	KCST_Return = 7,			// return TargetObject->TargetProperty
-	KCST_EndOfThread = 8,		// if (FlowStack.Num()) { NextState = FlowStack.Pop; } else { return; }
-	KCST_Comment = 9,			// /* Comment */
-	KCST_ComputedGoto = 10,		// NextState = LHS;
-	KCST_EndOfThreadIfNot = 11, // [if (!TargetObject->TargetProperty)] { same as KCST_EndOfThread; }
-	KCST_DebugSite = 12,		// NOP with recorded address
-	KCST_CastObjToInterface = 13,	// TargetInterface(TargetObject)
-	KCST_DynamicCast = 14,	    // Cast<TargetClass>(TargetObject)
-	KCST_ObjectToBool = 15,		// (TargetObject != None)
-	KCST_AddMulticastDelegate = 16,		// TargetDelegate->Add(EventDelegate)
-	KCST_ClearMulticastDelegate = 17,	// TargetDelegate->Clear()
-	KCST_WireTraceSite = 18,    // NOP with recorded address (never a step target)
-	KCST_BindDelegate = 19,		// Creates simple delegate
-	KCST_RemoveMulticastDelegate = 20,	//  TargetDelegate->Remove(EventDelegate)
-	KCST_CallDelegate = 21,		// TargetDelegate->Broadcast(...)
-	KCST_CreateArray = 22,		// Creates and sets an array literal term
-	KCST_CrossInterfaceCast = 23,	// TargetInterface(Interface)
-	KCST_MetaCast = 24,	    // Cast<TargetClass>(TargetObject)
-	KCST_AssignmentOnPersistentFrame = 25, //
-	KCST_CastInterfaceToObj = 26, // Cast<TargetClass>(TargetInterface)
-	KCST_GotoReturn = 27,	// goto ReturnLabel
-	KCST_GotoReturnIfNot = 28, // [if (!TargetObject->TargetProperty)] goto TargetLabel
+	// [wiring =] TargetObject->FunctionToCall(wiring)
+	KCST_CallFunction = 1,
+	// TargetObject->TargetProperty = [wiring]
+	KCST_Assignment = 2,
+	// One of the other types with a compilation error during statement generation
+	KCST_CompileError = 3,
+	// goto TargetLabel
+	KCST_UnconditionalGoto = 4,
+	// FlowStack.Push(TargetLabel)
+	KCST_PushState = 5,
+	// [if (!TargetObject->TargetProperty)] goto TargetLabel
+	KCST_GotoIfNot = 6,
+	// return TargetObject->TargetProperty
+	KCST_Return = 7,
+	// if (FlowStack.Num()) { NextState = FlowStack.Pop; } else { return; }
+	KCST_EndOfThread = 8,
+	// Comment
+	KCST_Comment = 9,
+	// NextState = LHS;
+	KCST_ComputedGoto = 10,
+	// [if (!TargetObject->TargetProperty)] { same as KCST_EndOfThread; }
+	KCST_EndOfThreadIfNot = 11,
+	// NOP with recorded address
+	KCST_DebugSite = 12,
+	// TargetInterface(TargetObject)
+	KCST_CastObjToInterface = 13,
+	// Cast<TargetClass>(TargetObject)
+	KCST_DynamicCast = 14,
+	// (TargetObject != None)
+	KCST_ObjectToBool = 15,
+	// TargetDelegate->Add(EventDelegate)
+	KCST_AddMulticastDelegate = 16,
+	// TargetDelegate->Clear()
+	KCST_ClearMulticastDelegate = 17,
+	// NOP with recorded address (never a step target)
+	KCST_WireTraceSite = 18,
+	// Creates simple delegate
+	KCST_BindDelegate = 19,
+	// TargetDelegate->Remove(EventDelegate)
+	KCST_RemoveMulticastDelegate = 20,
+	// TargetDelegate->Broadcast(...)
+	KCST_CallDelegate = 21,
+	// Creates and sets an array literal term
+	KCST_CreateArray = 22,
+	// TargetInterface(Interface)
+	KCST_CrossInterfaceCast = 23,
+	// Cast<TargetClass>(TargetObject)
+	KCST_MetaCast = 24,
+	KCST_AssignmentOnPersistentFrame = 25,
+	// Cast<TargetClass>(TargetInterface)
+	KCST_CastInterfaceToObj = 26,
+	// goto ReturnLabel
+	KCST_GotoReturn = 27,
+	// [if (!TargetObject->TargetProperty)] goto TargetLabel
+	KCST_GotoReturnIfNot = 28,
 	KCST_SwitchValue = 29,
-	// Kismet instrumentation extensions
-	KCST_InstrumentedEvent,				// Instrumented event
-	KCST_InstrumentedEventStop,			// Instrumented event stop
-	KCST_InstrumentedPureNodeEntry,		// Instrumented pure node entry
-	KCST_InstrumentedWireEntry,			// Instrumented wiretrace entry
-	KCST_InstrumentedWireExit,			// Instrumented wiretrace exit
-	KCST_InstrumentedStatePush,			// Instrumented state push
-	KCST_InstrumentedStateRestore,		// Instrumented state restore
-	KCST_InstrumentedStateReset,		// Instrumented state reset
-	KCST_InstrumentedStateSuspend,		// Instrumented state suspend
-	KCST_InstrumentedStatePop,			// Instrumented state pop
-	KCST_InstrumentedTunnelEndOfThread,	// Instrumented tunnel exit
-	//
+	
+	//~ Kismet instrumentation extensions:
+
+	// Instrumented event
+	KCST_InstrumentedEvent,
+	// Instrumented event stop
+	KCST_InstrumentedEventStop,
+	// Instrumented pure node entry
+	KCST_InstrumentedPureNodeEntry,
+	// Instrumented wiretrace entry
+	KCST_InstrumentedWireEntry,
+	// Instrumented wiretrace exit
+	KCST_InstrumentedWireExit,
+	// Instrumented state push
+	KCST_InstrumentedStatePush,
+	// Instrumented state restore
+	KCST_InstrumentedStateRestore,
+	// Instrumented state reset
+	KCST_InstrumentedStateReset,
+	// Instrumented state suspend
+	KCST_InstrumentedStateSuspend,
+	// Instrumented state pop
+	KCST_InstrumentedStatePop,
+	// Instrumented tunnel exit
+	KCST_InstrumentedTunnelEndOfThread,
+
 	KCST_ArrayGetByRef,
 	KCST_CreateSet,
 	KCST_CreateMap,
