Copying //UE4/Dev-Framework to //UE4/Dev-Main (Source: //UE4/Dev-Framework @ 3716594)

diff --git a/Engine/Source/Editor/KismetCompiler/Private/KismetCompiler.cpp b/Engine/Source/Editor/KismetCompiler/Private/KismetCompiler.cpp
--- a/Engine/Source/Editor/KismetCompiler/Private/KismetCompiler.cpp
+++ b/Engine/Source/Editor/KismetCompiler/Private/KismetCompiler.cpp
@@ -6,7 +6,6 @@
 
 
 #include "KismetCompiler.h"
-#include "AnimBlueprintCompiler.h"
 #include "Engine/BlueprintGeneratedClass.h"
 #include "Misc/CoreMisc.h"
 #include "Components/ActorComponent.h"
@@ -400,7 +399,7 @@ void FKismetCompilerContext::ValidateLink(const UEdGraphPin* PinA, const UEdGrap
 	{
 		const UEdGraphPin* InputPin = (EEdGraphPinDirection::EGPD_Input == PinA->Direction) ? PinA : PinB;
 		const UEdGraphPin* OutputPin = (EEdGraphPinDirection::EGPD_Output == PinA->Direction) ? PinA : PinB;
-		const bool bInvalidConnection = InputPin && OutputPin && (OutputPin->PinType.PinCategory == Schema->PC_Interface) && (InputPin->PinType.PinCategory == Schema->PC_Object);
+		const bool bInvalidConnection = InputPin && OutputPin && (OutputPin->PinType.PinCategory == UEdGraphSchema_K2::PC_Interface) && (InputPin->PinType.PinCategory == UEdGraphSchema_K2::PC_Object);
 		if (bInvalidConnection)
 		{
 			MessageLog.Error(*LOCTEXT("PinTypeMismatch_Error_UseExplictCast", "Can't connect pins @@ (Interface) and @@ (Object). Use an explicit cast node.").ToString(), OutputPin, InputPin);
@@ -684,12 +683,11 @@ void FKismetCompilerContext::CreateClassVariablesFromBlueprint()
 			continue;
 		}
 
-		FEdGraphPinType TimelinePinType(Schema->PC_Object, FString(), UTimelineComponent::StaticClass(), EPinContainerType::None, false, FEdGraphTerminalType());
+		FEdGraphPinType TimelinePinType(UEdGraphSchema_K2::PC_Object, NAME_None, UTimelineComponent::StaticClass(), EPinContainerType::None, false, FEdGraphTerminalType());
 
 		// Previously UTimelineComponent object has exactly the same name as UTimelineTemplate object (that obj was in blueprint)
 		const FString TimelineVariableName = UTimelineTemplate::TimelineTemplateNameToVariableName(Timeline->GetFName());
-		UProperty* TimelineProperty = CreateVariable(*TimelineVariableName, TimelinePinType);
-		if (TimelineProperty != NULL)
+		if (UProperty* TimelineProperty = CreateVariable(*TimelineVariableName, TimelinePinType))
 		{
 			TimelineProperty->SetMetaData( TEXT("Category"), *Blueprint->GetName() );
 			TimelineProperty->SetPropertyFlags(CPF_BlueprintVisible);
@@ -697,22 +695,22 @@ void FKismetCompilerContext::CreateClassVariablesFromBlueprint()
 			TimelineToMemberVariableMap.Add(Timeline, TimelineProperty);
 		}
 
-		FEdGraphPinType DirectionPinType(Schema->PC_Byte, FString(), FTimeline::GetTimelineDirectionEnum(), EPinContainerType::None, false, FEdGraphTerminalType());
+		FEdGraphPinType DirectionPinType(UEdGraphSchema_K2::PC_Byte, NAME_None, FTimeline::GetTimelineDirectionEnum(), EPinContainerType::None, false, FEdGraphTerminalType());
 		CreateVariable(Timeline->GetDirectionPropertyName(), DirectionPinType);
 
-		FEdGraphPinType FloatPinType(Schema->PC_Float, FString(), nullptr, EPinContainerType::None, false, FEdGraphTerminalType());
+		FEdGraphPinType FloatPinType(UEdGraphSchema_K2::PC_Float, NAME_None, nullptr, EPinContainerType::None, false, FEdGraphTerminalType());
 		for (const FTTFloatTrack& FloatTrack : Timeline->FloatTracks)
 		{
 			CreateVariable(Timeline->GetTrackPropertyName(FloatTrack.TrackName), FloatPinType);
 		}
 
-		FEdGraphPinType VectorPinType(Schema->PC_Struct, FString(), VectorStruct, EPinContainerType::None, false, FEdGraphTerminalType());
+		FEdGraphPinType VectorPinType(UEdGraphSchema_K2::PC_Struct, NAME_None, VectorStruct, EPinContainerType::None, false, FEdGraphTerminalType());
 		for (const FTTVectorTrack& VectorTrack : Timeline->VectorTracks)
 		{
 			CreateVariable(Timeline->GetTrackPropertyName(VectorTrack.TrackName), VectorPinType);
 		}
 
-		FEdGraphPinType LinearColorPinType(Schema->PC_Struct, FString(), LinearColorStruct, EPinContainerType::None, false, FEdGraphTerminalType());
+		FEdGraphPinType LinearColorPinType(UEdGraphSchema_K2::PC_Struct, NAME_None, LinearColorStruct, EPinContainerType::None, false, FEdGraphTerminalType());
 		for (const FTTLinearColorTrack& LinearColorTrack : Timeline->LinearColorTracks)
 		{
 			CreateVariable(Timeline->GetTrackPropertyName(LinearColorTrack.TrackName), LinearColorPinType);
@@ -735,7 +733,7 @@ void FKismetCompilerContext::CreateClassVariablesFromBlueprint()
 				FName VarName = Node->GetVariableName();
 				if ((VarName != NAME_None) && (Node->ComponentClass != nullptr))
 				{
-					FEdGraphPinType Type(Schema->PC_Object, FString(), Node->ComponentClass, EPinContainerType::None, false, FEdGraphTerminalType());
+					FEdGraphPinType Type(UEdGraphSchema_K2::PC_Object, NAME_None, Node->ComponentClass, EPinContainerType::None, false, FEdGraphTerminalType());
 					if (UProperty* NewProperty = CreateVariable(VarName, Type))
 					{
 						const FText CategoryName = Node->CategoryName.IsEmpty() ? FText::FromString(Blueprint->GetName()) : Node->CategoryName ;
@@ -1429,7 +1427,7 @@ void FKismetCompilerContext::ValidateNoWildcardPinsInGraph(const UEdGraph* Sourc
 			{
 				if (const UEdGraphPin* Pin = Node->Pins[PinIndex])
 				{
-					if (Pin->PinType.PinCategory == Schema->PC_Wildcard)
+					if (Pin->PinType.PinCategory == UEdGraphSchema_K2::PC_Wildcard)
 					{
 						// Wildcard pins should never be seen by the compiler; they should always be forced into a particular type by wiring.
 						MessageLog.Error(*LOCTEXT("UndeterminedPinType_Error", "The type of @@ is undetermined.  Connect something to @@ to imply a specific type.").ToString(), Pin, Pin->GetOwningNodeUnchecked());
@@ -1583,7 +1581,7 @@ void FKismetCompilerContext::PrecompileFunction(FKismetFunctionContext& Context,
 			
 			if (UEdGraphPin* WorldContextPin = Context.EntryPoint->GetAutoWorldContextPin())
 			{
-				Context.Function->SetMetaData(FBlueprintMetadata::MD_WorldContext, *WorldContextPin->PinName);
+				Context.Function->SetMetaData(FBlueprintMetadata::MD_WorldContext, *WorldContextPin->PinName.ToString());
 			}
 		}
 
@@ -1665,7 +1663,7 @@ void FKismetCompilerContext::PrecompileFunction(FKismetFunctionContext& Context,
 			!Context.IsEventGraph() &&
 			!Context.bIsSimpleStubGraphWithNoParams &&
 			Context.CanBeCalledByKismet() &&
-			Context.Function->GetFName() != Context.Schema->FN_UserConstructionScript)
+			Context.Function->GetFName() != UEdGraphSchema_K2::FN_UserConstructionScript)
 		{
 			// dig into the (actual) source graph and find the original return node:
 			UObject* Object = Context.MessageLog.FindSourceObject(Context.SourceGraph);
@@ -2077,7 +2075,7 @@ void FKismetCompilerContext::SetCalculatedMetaDataAndFlags(UFunction* Function,
 	}
 	if (UEdGraphPin* WorldContextPin = EntryNode->GetAutoWorldContextPin())
 	{
-		Function->SetMetaData(FBlueprintMetadata::MD_WorldContext, *WorldContextPin->PinName);
+		Function->SetMetaData(FBlueprintMetadata::MD_WorldContext, *WorldContextPin->PinName.ToString());
 	}
 
 	for (int32 EntryPinIndex = 0; EntryPinIndex < EntryNode->Pins.Num(); ++EntryPinIndex)
@@ -2090,7 +2088,7 @@ void FKismetCompilerContext::SetCalculatedMetaDataAndFlags(UFunction* Function,
 			(EntryPin->PinType.PinCategory != UEdGraphSchema_K2::PC_Interface) && 
 			!EntryPin->DefaultValue.IsEmpty() )
 		{
-			Function->SetMetaData(*EntryPin->PinName, *EntryPin->DefaultValue);
+			Function->SetMetaData(EntryPin->PinName, *EntryPin->DefaultValue);
 		}
 	}
 
@@ -2324,7 +2322,7 @@ void FKismetCompilerContext::BuildDynamicBindingObjects(UBlueprintGeneratedClass
  * @param PinName		The pin name to redirect output from, into the pin of the node event
  * @param ExecFuncName	The event signature name that the event node implements
  */ 
-void FKismetCompilerContext::CreatePinEventNodeForTimelineFunction(UK2Node_Timeline* TimelineNode, UEdGraph* SourceGraph, FName FunctionName, const FString& PinName, FName ExecFuncName)
+void FKismetCompilerContext::CreatePinEventNodeForTimelineFunction(UK2Node_Timeline* TimelineNode, UEdGraph* SourceGraph, FName FunctionName, const FName PinName, FName ExecFuncName)
 {
 	UEdGraphPin* SourcePin = nullptr;
 	if (UK2Node_Timeline* SourceNode = Cast<UK2Node_Timeline>(MessageLog.FindSourceObject(TimelineNode)))
@@ -2339,7 +2337,7 @@ void FKismetCompilerContext::CreatePinEventNodeForTimelineFunction(UK2Node_Timel
 
 	// Move any links from 'update' pin to the 'update event' node
 	UEdGraphPin* UpdatePin = TimelineNode ? TimelineNode->FindPin(PinName) : nullptr;
-	ensureMsgf(UpdatePin, TEXT("Timeline '%s' has no pin '%s'"), *GetPathNameSafe(TimelineNode), *PinName);
+	ensureMsgf(UpdatePin, TEXT("Timeline '%s' has no pin '%s'"), *GetPathNameSafe(TimelineNode), *PinName.ToString());
 
 	UEdGraphPin* UpdateOutput = Schema->FindExecutionPin(*TimelineEventNode, EGPD_Output);
 
@@ -2357,7 +2355,7 @@ UK2Node_CallFunction* FKismetCompilerContext::CreateCallTimelineFunction(UK2Node
 	CallNode->AllocateDefaultPins();
 
 	// Wire 'get timeline' to 'self' pin of function call
-	UEdGraphPin* CallSelfPin = CallNode->FindPinChecked(Schema->PN_Self);
+	UEdGraphPin* CallSelfPin = CallNode->FindPinChecked(UEdGraphSchema_K2::PN_Self);
 	TimelineVarPin->MakeLinkTo(CallSelfPin);
 
 	// Move any exec links from 'play' pin to the 'call play' node
@@ -2406,7 +2404,7 @@ void FKismetCompilerContext::ExpandTimelineNodes(UEdGraph* SourceGraph)
 				UClass* ValueClass = ValuePin->PinType.PinSubCategoryObject.IsValid() ? Cast<UClass>( ValuePin->PinType.PinSubCategoryObject.Get() ) : nullptr;
 				if( ValueClass == UTimelineComponent::StaticClass() )
 				{
-					const FName PinName( *ValuePin->PinName );
+					const FName PinName = ValuePin->PinName;
 					if( UTimelineTemplate* TimelineTemplate = Blueprint->FindTimelineTemplateByVariableName( PinName ))
 					{
 						TimelinePlayNodes.Add( PinName );
@@ -2423,8 +2421,6 @@ void FKismetCompilerContext::ExpandTimelineNodes(UEdGraph* SourceGraph)
 
 		if (bIsFullCompile)
 		{
-			const FString TimelineNameString = TimelineNode->TimelineName.ToString();
-
 			UEdGraphPin* PlayPin = TimelineNode->GetPlayPin();
 			bool bPlayPinConnected = (PlayPin->LinkedTo.Num() > 0);
 
@@ -2467,20 +2463,19 @@ void FKismetCompilerContext::ExpandTimelineNodes(UEdGraph* SourceGraph)
 				GetTimelineNode->AllocateDefaultPins();
 
 				// Debug data: Associate the timeline node instance with the property that was created earlier
-				UProperty* AssociatedTimelineInstanceProperty = TimelineToMemberVariableMap.FindChecked(Timeline);
-				if (AssociatedTimelineInstanceProperty != NULL)
+				if (UProperty* AssociatedTimelineInstanceProperty = TimelineToMemberVariableMap.FindChecked(Timeline))
 				{
 					UObject* TrueSourceObject = MessageLog.FindSourceObject(TimelineNode);
 					NewClass->GetDebugData().RegisterClassPropertyAssociation(TrueSourceObject, AssociatedTimelineInstanceProperty);
 				}
 
 				// Get the variable output pin
-				UEdGraphPin* TimelineVarPin = GetTimelineNode->FindPin(TimelineNameString);
+				UEdGraphPin* TimelineVarPin = GetTimelineNode->FindPin(TimelineNode->TimelineName);
 
 				// This might fail if this is the first compile after adding the timeline (property doesn't exist yet) - in that case, manually add the output pin
-				if (TimelineVarPin == NULL)
+				if (TimelineVarPin == nullptr)
 				{
-					TimelineVarPin = GetTimelineNode->CreatePin(EGPD_Output, Schema->PC_Object, FString(), UTimelineComponent::StaticClass(), TimelineNode->TimelineName.ToString());
+					TimelineVarPin = GetTimelineNode->CreatePin(EGPD_Output, UEdGraphSchema_K2::PC_Object, UTimelineComponent::StaticClass(), TimelineNode->TimelineName);
 				}
 
 				if (bPlayPinConnected)
@@ -2535,8 +2530,8 @@ void FKismetCompilerContext::ExpandTimelineNodes(UEdGraph* SourceGraph)
 		// Create event nodes for any event tracks
 		for(int32 EventTrackIdx=0; EventTrackIdx<Timeline->EventTracks.Num(); EventTrackIdx++)
 		{
-			FName EventTrackName =  Timeline->EventTracks[EventTrackIdx].TrackName;
-			CreatePinEventNodeForTimelineFunction(TimelineNode, SourceGraph, Timeline->GetEventTrackFunctionName(EventTrackIdx), EventTrackName.ToString(), EventSigFunc->GetFName());
+			const FName EventTrackName = Timeline->EventTracks[EventTrackIdx].TrackName;
+			CreatePinEventNodeForTimelineFunction(TimelineNode, SourceGraph, Timeline->GetEventTrackFunctionName(EventTrackIdx), EventTrackName, EventSigFunc->GetFName());
 		}
 
 		// Generate Update Pin Event Node
@@ -2581,16 +2576,18 @@ FPinConnectionResponse FKismetCompilerContext::CopyPinLinksToIntermediate(UEdGra
 	return ConnectionResult;
 }
 
-UK2Node_TemporaryVariable* FKismetCompilerContext::SpawnInternalVariable(UEdGraphNode* SourceNode, FString Category, FString SubCategory, UObject* SubcategoryObject, bool bIsArray, bool bIsSet, bool bIsMap, const FEdGraphTerminalType& ValueTerminalType)
+UK2Node_TemporaryVariable* FKismetCompilerContext::SpawnInternalVariable(UEdGraphNode* SourceNode, const FString& Category, const FString& SubCategory, UObject* SubcategoryObject, bool bIsArray, bool bIsSet, bool bIsMap, const FEdGraphTerminalType& ValueTerminalType)
 {
-	return SpawnInternalVariable(SourceNode, MoveTemp(Category), MoveTemp(SubCategory), SubcategoryObject, FEdGraphPinType::ToPinContainerType(bIsArray, bIsSet, bIsMap), ValueTerminalType);
+	PRAGMA_DISABLE_DEPRECATION_WARNINGS
+	return SpawnInternalVariable(SourceNode, Category, SubCategory, SubcategoryObject, FEdGraphPinType::ToPinContainerType(bIsArray, bIsSet, bIsMap), ValueTerminalType);
+	PRAGMA_ENABLE_DEPRECATION_WARNINGS
 }
 
-UK2Node_TemporaryVariable* FKismetCompilerContext::SpawnInternalVariable(UEdGraphNode* SourceNode, FString Category, FString SubCategory, UObject* SubcategoryObject, EPinContainerType PinContainerType, const FEdGraphTerminalType& ValueTerminalType)
+UK2Node_TemporaryVariable* FKismetCompilerContext::SpawnInternalVariable(UEdGraphNode* SourceNode, const FName Category, const FName SubCategory, UObject* SubcategoryObject, EPinContainerType PinContainerType, const FEdGraphTerminalType& ValueTerminalType)
 {
 	UK2Node_TemporaryVariable* Result = SpawnIntermediateNode<UK2Node_TemporaryVariable>(SourceNode);
 
-	Result->VariableType = FEdGraphPinType(MoveTemp(Category), MoveTemp(SubCategory), SubcategoryObject, PinContainerType, false, ValueTerminalType);
+	Result->VariableType = FEdGraphPinType(Category, SubCategory, SubcategoryObject, PinContainerType, false, ValueTerminalType);
 	Result->AllocateDefaultPins();
 
 	return Result;
@@ -2731,7 +2728,7 @@ void FKismetCompilerContext::CreateFunctionStubForEvent(UK2Node_Event* SrcEventN
 			UEdGraphPin* UGSourcePin = SrcEventNode->FindPin(SourcePin->PinName);
 			const FString MemberVariableName = ClassScopeNetNameMap.MakeValidName(UGSourcePin);
 
-			UEdGraphPin* DestPin = AssignmentNode->CreatePin(EGPD_Input, SourcePin->PinType, MemberVariableName);
+			UEdGraphPin* DestPin = AssignmentNode->CreatePin(EGPD_Input, SourcePin->PinType, *MemberVariableName);
 			MessageLog.NotifyIntermediatePinCreation(DestPin, SourcePin);
 			DestPin->MakeLinkTo(SourcePin);
 		}
@@ -2748,16 +2745,16 @@ void FKismetCompilerContext::CreateFunctionStubForEvent(UK2Node_Event* SrcEventN
 	CallIntoUbergraph->NodePosX = 300;
 
 	// Use the ExecuteUbergraph base function to generate the pins...
-	CallIntoUbergraph->FunctionReference.SetExternalMember(Schema->FN_ExecuteUbergraphBase, UObject::StaticClass());
+	CallIntoUbergraph->FunctionReference.SetExternalMember(UEdGraphSchema_K2::FN_ExecuteUbergraphBase, UObject::StaticClass());
 	CallIntoUbergraph->AllocateDefaultPins();
 	
 	// ...then swap to the generated version for this level
 	CallIntoUbergraph->FunctionReference.SetSelfMember(GetUbergraphCallName());
 	UEdGraphPin* CallIntoUbergraphSelf = Schema->FindSelfPin(*CallIntoUbergraph, EGPD_Input);
-	CallIntoUbergraphSelf->PinType.PinSubCategory = Schema->PSC_Self;
+	CallIntoUbergraphSelf->PinType.PinSubCategory = UEdGraphSchema_K2::PSC_Self;
 	CallIntoUbergraphSelf->PinType.PinSubCategoryObject = *Blueprint->SkeletonGeneratedClass;
 
-	UEdGraphPin* EntryPointPin = CallIntoUbergraph->FindPin(Schema->PN_EntryPoint);
+	UEdGraphPin* EntryPointPin = CallIntoUbergraph->FindPin(UEdGraphSchema_K2::PN_EntryPoint);
 	if (EntryPointPin)
 	{
 		EntryPointPin->DefaultValue = TEXT("0");
@@ -2770,7 +2767,7 @@ void FKismetCompilerContext::CreateFunctionStubForEvent(UK2Node_Event* SrcEventN
 	UEdGraphPin* ExecEntryOut = Schema->FindExecutionPin(*EntryNode, EGPD_Output);
 	UEdGraphPin* ExecCallIn = Schema->FindExecutionPin(*CallIntoUbergraph, EGPD_Input);
 
-	if (AssignmentNode != NULL)
+	if (AssignmentNode)
 	{
 		UEdGraphPin* ExecVariablesIn = Schema->FindExecutionPin(*AssignmentNode, EGPD_Input);
 		UEdGraphPin* ExecVariablesOut = Schema->FindExecutionPin(*AssignmentNode, EGPD_Output);
@@ -3079,7 +3076,7 @@ void FKismetCompilerContext::CreateAndProcessUbergraph()
 		{
 			UK2Node_FunctionEntry* EntryNode = SpawnIntermediateNode<UK2Node_FunctionEntry>(NULL, ConsolidatedEventGraph);
 			EntryNode->SignatureClass = UObject::StaticClass();
-			EntryNode->SignatureName = Schema->FN_ExecuteUbergraphBase;
+			EntryNode->SignatureName = UEdGraphSchema_K2::FN_ExecuteUbergraphBase;
 			EntryNode->CustomGeneratedFunctionName = ConsolidatedEventGraph->GetFName();
 			EntryNode->AllocateDefaultPins();
 		}
@@ -3245,7 +3242,7 @@ void FKismetCompilerContext::ExpandTunnelsAndMacros(UEdGraph* SourceGraph)
 			TArray<UEdGraphNode*> MacroNodes(ClonedGraph->Nodes);
 
 			// resolve any wildcard pins in the nodes cloned from the macro
-			if (!MacroInstanceNode->ResolvedWildcardType.PinCategory.IsEmpty())
+			if (!MacroInstanceNode->ResolvedWildcardType.PinCategory.IsNone())
 			{
 				for (UEdGraphNode* const ClonedNode : ClonedGraph->Nodes)
 				{
@@ -3253,7 +3250,7 @@ void FKismetCompilerContext::ExpandTunnelsAndMacros(UEdGraph* SourceGraph)
 					{
 						for (UEdGraphPin* const ClonedPin : ClonedNode->Pins)
 						{
-							if ( ClonedPin && (ClonedPin->PinType.PinCategory == Schema->PC_Wildcard) )
+							if ( ClonedPin && (ClonedPin->PinType.PinCategory == UEdGraphSchema_K2::PC_Wildcard) )
 							{
 								// copy only type info, so array or ref status is preserved
 								ClonedPin->PinType.PinCategory = MacroInstanceNode->ResolvedWildcardType.PinCategory;
@@ -3298,7 +3295,7 @@ void FKismetCompilerContext::ExpandTunnelsAndMacros(UEdGraph* SourceGraph)
 					else if (Pin->LinkedTo.Num() == 0 &&
 							Pin->Direction == EGPD_Input &&
 							Pin->DefaultValue != FString() &&
-							Pin->PinType.PinCategory == Schema->PC_Byte &&
+							Pin->PinType.PinCategory == UEdGraphSchema_K2::PC_Byte &&
 							Pin->PinType.PinSubCategoryObject.IsValid() &&
 							Pin->PinType.PinSubCategoryObject->IsA<UEnum>())
 					{
@@ -3306,7 +3303,7 @@ void FKismetCompilerContext::ExpandTunnelsAndMacros(UEdGraph* SourceGraph)
 						UK2Node_EnumLiteral* EnumLiteralNode = SpawnIntermediateNode<UK2Node_EnumLiteral>(MacroInstanceNode, SourceGraph);
 						EnumLiteralNode->Enum = CastChecked<UEnum>(Pin->PinType.PinSubCategoryObject.Get());
 						EnumLiteralNode->AllocateDefaultPins();
-						EnumLiteralNode->FindPinChecked(Schema->PN_ReturnValue)->MakeLinkTo(Pin);
+						EnumLiteralNode->FindPinChecked(UEdGraphSchema_K2::PN_ReturnValue)->MakeLinkTo(Pin);
 
 						UEdGraphPin* InPin = EnumLiteralNode->FindPinChecked(UK2Node_EnumLiteral::GetEnumInputPinName());
 						check(InPin);
@@ -3396,7 +3393,7 @@ void FKismetCompilerContext::ExpandTunnelsAndMacros(UEdGraph* SourceGraph)
 			UEdGraphNode* InputSink = TunnelNode->GetInputSink();
 			for (UEdGraphPin const* TunnelPin : TunnelNode->Pins)
 			{
-				if ((TunnelPin->Direction != EGPD_Input) || (TunnelPin->PinType.PinCategory != Schema->PC_Exec))
+				if ((TunnelPin->Direction != EGPD_Input) || (TunnelPin->PinType.PinCategory != UEdGraphSchema_K2::PC_Exec))
 				{
 					continue;
 				}
@@ -3445,6 +3442,12 @@ void FKismetCompilerContext::ProcessOneFunctionGraph(UEdGraph* SourceGraph, bool
 {
 	BP_SCOPED_COMPILER_EVENT_STAT(EKismetCompilerStats_ProcessFunctionGraph);
 
+	if (SourceGraph->GetFName() == Schema->FN_UserConstructionScript && FBlueprintEditorUtils::IsDataOnlyBlueprint(Blueprint))
+	{
+		// This is a data only blueprint, we do not want to actually create our user construction script as it only consists of a call to the parent
+		return;
+	}
+
 	// Clone the source graph so we can modify it as needed; merging in the child graphs
 	UEdGraph* FunctionGraph = FEdGraphUtilities::CloneGraph(SourceGraph, Blueprint, &MessageLog, true); 
 	FEdGraphUtilities::MergeChildrenGraphsIn(FunctionGraph, FunctionGraph, /*bRequireSchemaMatch=*/ true);
@@ -3510,7 +3513,7 @@ void FKismetCompilerContext::ValidateFunctionGraphNames()
 		for (int32 FunctionIndex=0; FunctionIndex < Blueprint->FunctionGraphs.Num(); ++FunctionIndex)
 		{
 			UEdGraph* FunctionGraph = Blueprint->FunctionGraphs[FunctionIndex];
-			if(FunctionGraph->GetFName() != Schema->FN_UserConstructionScript)
+			if(FunctionGraph->GetFName() != UEdGraphSchema_K2::FN_UserConstructionScript)
 			{
 				if( ParentBPNameValidator->IsValid(FunctionGraph->GetName()) != EValidatorResult::Ok )
 				{
@@ -3807,7 +3810,7 @@ void FKismetCompilerContext::CompileClassLayout(EInternalCompilerFlags InternalF
 	if (UsePersistentUberGraphFrame() && UbergraphContext)
 	{
 		//UBER GRAPH PERSISTENT FRAME
-		FEdGraphPinType Type(TEXT("struct"), FString(), FPointerToUberGraphFrame::StaticStruct(), EPinContainerType::None, false, FEdGraphTerminalType());
+		FEdGraphPinType Type(TEXT("struct"), NAME_None, FPointerToUberGraphFrame::StaticStruct(), EPinContainerType::None, false, FEdGraphTerminalType());
 		UProperty* Property = CreateVariable(UBlueprintGeneratedClass::GetUberGraphFrameName(), Type);
 		Property->SetPropertyFlags(CPF_DuplicateTransient | CPF_Transient);
 	}
@@ -4473,17 +4476,5 @@ FString FKismetCompilerContext::GetGuid(const UEdGraphNode* Node) const
 	return Ret.ToString();
 }
 
-TSharedPtr<FKismetCompilerContext> FKismetCompilerContext::GetCompilerForBP(UBlueprint* BP, FCompilerResultsLog& InMessageLog, const FKismetCompilerOptions& InCompileOptions)
-{
-	if(UAnimBlueprint* AnimBP = Cast<UAnimBlueprint>(BP))
-	{
-		return TSharedPtr<FKismetCompilerContext>(new FAnimBlueprintCompiler(AnimBP, InMessageLog, InCompileOptions, nullptr));
-	}
-	else
-	{
-		return TSharedPtr<FKismetCompilerContext>(new FKismetCompilerContext(BP, InMessageLog, InCompileOptions, nullptr));
-	}
-}
-
 #undef LOCTEXT_NAMESPACE
 //////////////////////////////////////////////////////////////////////////
diff --git a/Engine/Source/Editor/KismetCompiler/Private/KismetCompilerMisc.cpp b/Engine/Source/Editor/KismetCompiler/Private/KismetCompilerMisc.cpp
--- a/Engine/Source/Editor/KismetCompiler/Private/KismetCompilerMisc.cpp
+++ b/Engine/Source/Editor/KismetCompiler/Private/KismetCompilerMisc.cpp
@@ -5,6 +5,8 @@
 =============================================================================*/
 
 #include "KismetCompilerMisc.h"
+
+#include "BlueprintCompilationManager.h"
 #include "Misc/CoreMisc.h"
 #include "UObject/MetaData.h"
 #include "UObject/UnrealType.h"
@@ -47,34 +49,33 @@ DECLARE_CYCLE_STAT(TEXT("Resolve compiled statements"), EKismetCompilerStats_Res
 //////////////////////////////////////////////////////////////////////////
 // FKismetCompilerUtilities
 
-static bool IsTypeCompatibleWithProperty(UEdGraphPin* SourcePin, const FEdGraphPinType& OwningType, const FEdGraphTerminalType& TerminalType, UProperty* TestProperty, FCompilerResultsLog& MessageLog, const UEdGraphSchema_K2* Schema, UClass* SelfClass)
+static bool IsTypeCompatibleWithProperty(UEdGraphPin* SourcePin, const FEdGraphPinType& OwningType, const FEdGraphTerminalType& TerminalType, UProperty* TestProperty, FCompilerResultsLog& MessageLog, UClass* SelfClass)
 {
-	check(SourcePin != NULL);
+	check(SourcePin);
 	const EEdGraphPinDirection Direction = SourcePin->Direction; 
-	const FString& PinCategory = TerminalType.TerminalCategory;
-	const FString& PinSubCategory = TerminalType.TerminalSubCategory;
+	const FName PinCategory = TerminalType.TerminalCategory;
+	const FName PinSubCategory = TerminalType.TerminalSubCategory;
 	const UObject* PinSubCategoryObject = TerminalType.TerminalSubCategoryObject.Get();
 	
 	const UFunction* OwningFunction = Cast<UFunction>(TestProperty->GetOuter());
 
 	bool bTypeMismatch = false;
 	bool bSubtypeMismatch = false;
-	FString DesiredSubType(TEXT(""));
 
-	if (PinCategory == Schema->PC_Boolean)
+	if (PinCategory == UEdGraphSchema_K2::PC_Boolean)
 	{
 		UBoolProperty* SpecificProperty = Cast<UBoolProperty>(TestProperty);
-		bTypeMismatch = (SpecificProperty == NULL);
+		bTypeMismatch = (SpecificProperty == nullptr);
 	}
-	else if (PinCategory == Schema->PC_Byte)
+	else if (PinCategory == UEdGraphSchema_K2::PC_Byte)
 	{
 		UByteProperty* ByteProperty = Cast<UByteProperty>(TestProperty);
 		UEnumProperty* EnumProperty = Cast<UEnumProperty>(TestProperty);
 		bTypeMismatch = (ByteProperty == nullptr) && (EnumProperty == nullptr || !EnumProperty->GetUnderlyingProperty()->IsA<UByteProperty>());
 	}
-	else if ((PinCategory == Schema->PC_Class) || (PinCategory == Schema->PC_SoftClass))
+	else if ((PinCategory == UEdGraphSchema_K2::PC_Class) || (PinCategory == UEdGraphSchema_K2::PC_SoftClass))
 	{
-		const UClass* ClassType = (PinSubCategory == Schema->PSC_Self) ? SelfClass : Cast<const UClass>(PinSubCategoryObject);
+		const UClass* ClassType = (PinSubCategory == UEdGraphSchema_K2::PSC_Self) ? SelfClass : Cast<const UClass>(PinSubCategoryObject);
 
 		if (ClassType == NULL)
 		{
@@ -94,15 +95,13 @@ static bool IsTypeCompatibleWithProperty(UEdGraphPin* SourcePin, const FEdGraphP
 
 			if (MetaClass != NULL)
 			{
-				DesiredSubType = MetaClass->GetName();
-
 				const UClass* OutputClass = (Direction == EGPD_Output) ? ClassType :  MetaClass;
 				const UClass* InputClass = (Direction == EGPD_Output) ? MetaClass : ClassType;
 
 				// It matches if it's an exact match or if the output class is more derived than the input class
 				bTypeMismatch = bSubtypeMismatch = !((OutputClass == InputClass) || (OutputClass->IsChildOf(InputClass)));
 
-				if ((PinCategory == Schema->PC_SoftClass) && (!TestProperty->IsA<USoftClassProperty>()))
+				if ((PinCategory == UEdGraphSchema_K2::PC_SoftClass) && (!TestProperty->IsA<USoftClassProperty>()))
 				{
 					bTypeMismatch = true;
 				}
@@ -113,22 +112,22 @@ static bool IsTypeCompatibleWithProperty(UEdGraphPin* SourcePin, const FEdGraphP
 			}
 		}
 	}
-	else if (PinCategory == Schema->PC_Float)
+	else if (PinCategory == UEdGraphSchema_K2::PC_Float)
 	{
 		UFloatProperty* SpecificProperty = Cast<UFloatProperty>(TestProperty);
-		bTypeMismatch = (SpecificProperty == NULL);
+		bTypeMismatch = (SpecificProperty == nullptr);
 	}
-	else if (PinCategory == Schema->PC_Int)
+	else if (PinCategory == UEdGraphSchema_K2::PC_Int)
 	{
 		UIntProperty* SpecificProperty = Cast<UIntProperty>(TestProperty);
-		bTypeMismatch = (SpecificProperty == NULL);
+		bTypeMismatch = (SpecificProperty == nullptr);
 	}
-	else if (PinCategory == Schema->PC_Name)
+	else if (PinCategory == UEdGraphSchema_K2::PC_Name)
 	{
 		UNameProperty* SpecificProperty = Cast<UNameProperty>(TestProperty);
-		bTypeMismatch = (SpecificProperty == NULL);
+		bTypeMismatch = (SpecificProperty == nullptr);
 	}
-	else if (PinCategory == Schema->PC_Delegate)
+	else if (PinCategory == UEdGraphSchema_K2::PC_Delegate)
 	{
 		const UFunction* SignatureFunction = FMemberReference::ResolveSimpleMemberReference<UFunction>(OwningType.PinSubCategoryMemberReference);
 		const UDelegateProperty* PropertyDelegate = Cast<const UDelegateProperty>(TestProperty);
@@ -137,9 +136,9 @@ static bool IsTypeCompatibleWithProperty(UEdGraphPin* SourcePin, const FEdGraphP
 			&& PropertyDelegate->SignatureFunction 
 			&& PropertyDelegate->SignatureFunction->IsSignatureCompatibleWith(SignatureFunction));
 	}
-	else if ((PinCategory == Schema->PC_Object) || (PinCategory == Schema->PC_Interface) || (PinCategory == Schema->PC_SoftObject))
+	else if ((PinCategory == UEdGraphSchema_K2::PC_Object) || (PinCategory == UEdGraphSchema_K2::PC_Interface) || (PinCategory == UEdGraphSchema_K2::PC_SoftObject))
 	{
-		const UClass* ObjectType = (PinSubCategory == Schema->PSC_Self) ? SelfClass : Cast<const UClass>(PinSubCategoryObject);
+		const UClass* ObjectType = (PinSubCategory == UEdGraphSchema_K2::PSC_Self) ? SelfClass : Cast<const UClass>(PinSubCategoryObject);
 
 		if (ObjectType == NULL)
 		{
@@ -150,8 +149,6 @@ static bool IsTypeCompatibleWithProperty(UEdGraphPin* SourcePin, const FEdGraphP
 			UObjectPropertyBase* ObjProperty = Cast<UObjectPropertyBase>(TestProperty);
 			if (ObjProperty != NULL && ObjProperty->PropertyClass)
 			{
-				DesiredSubType = ObjProperty->PropertyClass->GetName();
-
 				const UClass* OutputClass = (Direction == EGPD_Output) ? ObjectType : ObjProperty->PropertyClass;
 				const UClass* InputClass  = (Direction == EGPD_Output) ? ObjProperty->PropertyClass : ObjectType;
 
@@ -184,7 +181,7 @@ static bool IsTypeCompatibleWithProperty(UEdGraphPin* SourcePin, const FEdGraphP
 				// It matches if it's an exact match or if the output class is more derived than the input class
 				bTypeMismatch = bSubtypeMismatch = !((OutputClass == InputClass) || (OutputClass->IsChildOf(InputClass)));
 
-				if ((PinCategory == Schema->PC_SoftObject) && (!TestProperty->IsA<USoftObjectProperty>()))
+				if ((PinCategory == UEdGraphSchema_K2::PC_SoftObject) && (!TestProperty->IsA<USoftObjectProperty>()))
 				{
 					bTypeMismatch = true;
 				}
@@ -198,7 +195,6 @@ static bool IsTypeCompatibleWithProperty(UEdGraphPin* SourcePin, const FEdGraphP
 				}
 				else 
 				{
-					DesiredSubType = InterfaceClass->GetName();
 					bTypeMismatch = ObjectType->ImplementsInterface(InterfaceClass);
 				}
 			}
@@ -208,17 +204,17 @@ static bool IsTypeCompatibleWithProperty(UEdGraphPin* SourcePin, const FEdGraphP
 			}
 		}
 	}
-	else if (PinCategory == Schema->PC_String)
+	else if (PinCategory == UEdGraphSchema_K2::PC_String)
 	{
 		UStrProperty* SpecificProperty = Cast<UStrProperty>(TestProperty);
-		bTypeMismatch = (SpecificProperty == NULL);
+		bTypeMismatch = (SpecificProperty == nullptr);
 	}
-	else if (PinCategory == Schema->PC_Text)
+	else if (PinCategory == UEdGraphSchema_K2::PC_Text)
 	{
 		UTextProperty* SpecificProperty = Cast<UTextProperty>(TestProperty);
-		bTypeMismatch = (SpecificProperty == NULL);
+		bTypeMismatch = (SpecificProperty == nullptr);
 	}
-	else if (PinCategory == Schema->PC_Struct)
+	else if (PinCategory == UEdGraphSchema_K2::PC_Struct)
 	{
 		const UScriptStruct* StructType = Cast<const UScriptStruct>(PinSubCategoryObject);
 		if (StructType == NULL)
@@ -230,7 +226,6 @@ static bool IsTypeCompatibleWithProperty(UEdGraphPin* SourcePin, const FEdGraphP
 			UStructProperty* StructProperty = Cast<UStructProperty>(TestProperty);
 			if (StructProperty != NULL)
 			{
-				DesiredSubType = StructProperty->Struct->GetName();
 				bool bMatchingStructs = (StructType == StructProperty->Struct);
 				if (auto UserDefinedStructFromProperty = Cast<const UUserDefinedStruct>(StructProperty->Struct))
 				{
@@ -267,8 +262,8 @@ bool FKismetCompilerUtilities::IsTypeCompatibleWithProperty(UEdGraphPin* SourceP
 	const FEdGraphPinType& Type = SourcePin->PinType;
 	const EEdGraphPinDirection Direction = SourcePin->Direction; 
 
-	const FString& PinCategory = Type.PinCategory;
-	const FString& PinSubCategory = Type.PinSubCategory;
+	const FName PinCategory = Type.PinCategory;
+	const FName PinSubCategory = Type.PinSubCategory;
 	const UObject* PinSubCategoryObject = Type.PinSubCategoryObject.Get();
 
 	UProperty* TestProperty = NULL;
@@ -289,19 +284,25 @@ bool FKismetCompilerUtilities::IsTypeCompatibleWithProperty(UEdGraphPin* SourceP
 				TArray<FString> ArrayPinComboNames;
 				ArrayPointerMetaData.ParseIntoArray(ArrayPinComboNames, TEXT(","), true);
 
-				for(auto Iter = ArrayPinComboNames.CreateConstIterator(); Iter; ++Iter)
+				if (ArrayPinComboNames.Num() > 0)
 				{
 					TArray<FString> ArrayPinNames;
-					Iter->ParseIntoArray(ArrayPinNames, TEXT("|"), true);
+					const FString SourcePinName = SourcePin->PinName.ToString();
 
-					if( ArrayPinNames[0] == SourcePin->PinName )
+					for (const FString& ArrayPinComboName : ArrayPinComboNames)
 					{
-						return true;
+						ArrayPinNames.Reset();
+						ArrayPinComboName.ParseIntoArray(ArrayPinNames, TEXT("|"), true);
+
+						if (ArrayPinNames[0] == SourcePinName)
+						{
+							return true;
+						}
 					}
 				}
 			}
 			
-			bTypeMismatch = ::IsTypeCompatibleWithProperty(SourcePin, Type, SourcePin->GetPrimaryTerminalType(), ArrayProp->Inner, MessageLog, Schema, SelfClass);
+			bTypeMismatch = ::IsTypeCompatibleWithProperty(SourcePin, Type, SourcePin->GetPrimaryTerminalType(), ArrayProp->Inner, MessageLog, SelfClass);
 		}
 		else
 		{
@@ -318,7 +319,7 @@ bool FKismetCompilerUtilities::IsTypeCompatibleWithProperty(UEdGraphPin* SourceP
 				return true;
 			}
 
-			bTypeMismatch = ::IsTypeCompatibleWithProperty(SourcePin, Type, SourcePin->GetPrimaryTerminalType(), SetProperty->ElementProp, MessageLog, Schema, SelfClass);
+			bTypeMismatch = ::IsTypeCompatibleWithProperty(SourcePin, Type, SourcePin->GetPrimaryTerminalType(), SetProperty->ElementProp, MessageLog, SelfClass);
 		}
 		else
 		{
@@ -335,8 +336,8 @@ bool FKismetCompilerUtilities::IsTypeCompatibleWithProperty(UEdGraphPin* SourceP
 				return true;
 			}
 
-			bTypeMismatch = ::IsTypeCompatibleWithProperty(SourcePin, Type, SourcePin->GetPrimaryTerminalType(), MapProperty->KeyProp, MessageLog, Schema, SelfClass);
-			bTypeMismatch = bTypeMismatch && ::IsTypeCompatibleWithProperty(SourcePin, Type, Type.PinValueType, MapProperty->ValueProp, MessageLog, Schema, SelfClass);
+			bTypeMismatch = ::IsTypeCompatibleWithProperty(SourcePin, Type, SourcePin->GetPrimaryTerminalType(), MapProperty->KeyProp, MessageLog, SelfClass);
+			bTypeMismatch = bTypeMismatch && ::IsTypeCompatibleWithProperty(SourcePin, Type, Type.PinValueType, MapProperty->ValueProp, MessageLog, SelfClass);
 		}
 		else
 		{
@@ -347,7 +348,7 @@ bool FKismetCompilerUtilities::IsTypeCompatibleWithProperty(UEdGraphPin* SourceP
 	else
 	{
 		// For scalars, we just take the passed in property
-		bTypeMismatch = ::IsTypeCompatibleWithProperty(SourcePin, Type, SourcePin->GetPrimaryTerminalType(), Property, MessageLog, Schema, SelfClass);
+		bTypeMismatch = ::IsTypeCompatibleWithProperty(SourcePin, Type, SourcePin->GetPrimaryTerminalType(), Property, MessageLog, SelfClass);
 	}
 
 	// Check for the early out...if this is a type dependent parameter in an array function
@@ -359,7 +360,7 @@ bool FKismetCompilerUtilities::IsTypeCompatibleWithProperty(UEdGraphPin* SourceP
 			const FString& DependentParams = OwningFunction->GetMetaData(FBlueprintMetadata::MD_ArrayDependentParam);
 			TArray<FString>	DependentParamNames;
 			DependentParams.ParseIntoArray(DependentParamNames, TEXT(","), true);
-			if (DependentParamNames.Find(SourcePin->PinName) != INDEX_NONE)
+			if (DependentParamNames.Find(SourcePin->PinName.ToString()) != INDEX_NONE)
 			{
 				//@todo:  This assumes that the wildcard coercion has done its job...I'd feel better if there was some easier way of accessing the target array type
 				return true;
@@ -407,20 +408,20 @@ bool FKismetCompilerUtilities::IsTypeCompatibleWithProperty(UEdGraphPin* SourceP
 			MessageLog.Error(*FString::Printf(*LOCTEXT("DirectionMismatchParameter_Error", "The direction of @@ doesn't match the direction of parameter %s").ToString(), *Property->GetName()), SourcePin);
 		}
 
- 		if (Property->HasAnyPropertyFlags(CPF_ReferenceParm))
- 		{
+		if (Property->HasAnyPropertyFlags(CPF_ReferenceParm)
+			&& (SourcePin->LinkedTo.Num() == 0)
+			&& (SourcePin->PinType.PinSubCategoryObject.Get() != TBaseStructure<FTransform>::Get()))
+		{
 			TArray<FString> AutoEmittedTerms;
 			Schema->GetAutoEmitTermParameters(OwningFunction, AutoEmittedTerms);
-			const bool bIsAutoEmittedTerm = AutoEmittedTerms.Contains(SourcePin->PinName);
+			const bool bIsAutoEmittedTerm = AutoEmittedTerms.Contains(SourcePin->PinName.ToString());
 
 			// Make sure reference parameters are linked, except for FTransforms, which have a special node handler that adds an internal constant term
- 			if (!bIsAutoEmittedTerm
-				&& (SourcePin->LinkedTo.Num() == 0)
-				&& (!SourcePin->PinType.PinSubCategoryObject.IsValid() || SourcePin->PinType.PinSubCategoryObject.Get()->GetName() != TEXT("Transform")) )
- 			{
- 				MessageLog.Error(*LOCTEXT("PassLiteral_Error", "Cannot pass a literal to @@.  Connect a variable to it instead.").ToString(), SourcePin);
- 			}
- 		}
+			if (!bIsAutoEmittedTerm)
+			{
+				MessageLog.Error(*LOCTEXT("PassLiteral_Error", "Cannot pass a literal to @@.  Connect a variable to it instead.").ToString(), SourcePin);
+			}
+		}
 	}
 
 	return NumErrorsAtStart == MessageLog.NumErrors;
@@ -537,7 +538,7 @@ UProperty* FKismetCompilerUtilities::FindPropertyInScope(UStruct* Scope, UEdGrap
 		{
 			UProperty* Property = *It;
 
-			if (Property->GetName() == Pin->PinName)
+			if (Property->GetFName() == Pin->PinName)
 			{
 				if (FKismetCompilerUtilities::IsTypeCompatibleWithProperty(Pin, Property, MessageLog, Schema, SelfClass))
 				{
@@ -673,24 +674,24 @@ bool FKismetCompilerUtilities::ValidateSelfCompatibility(const UEdGraphPin* Pin,
 	else
 	{
 		FEdGraphPinType SelfType;
-		SelfType.PinCategory = K2Schema->PC_Object;
-		SelfType.PinSubCategory = K2Schema->PSC_Self;
+		SelfType.PinCategory = UEdGraphSchema_K2::PC_Object;
+		SelfType.PinSubCategory = UEdGraphSchema_K2::PSC_Self;
 
 		if (!K2Schema->ArePinTypesCompatible(SelfType, Pin->PinType, BPClass))
 		{
-			FString PinType = Pin->PinType.PinCategory;
-			if ((Pin->PinType.PinCategory == K2Schema->PC_Object) ||
-				(Pin->PinType.PinCategory == K2Schema->PC_Interface) ||
-				(Pin->PinType.PinCategory == K2Schema->PC_Class))
+			FString PinType;
+			if ((Pin->PinType.PinCategory == UEdGraphSchema_K2::PC_Object) ||
+				(Pin->PinType.PinCategory == UEdGraphSchema_K2::PC_Interface) ||
+				(Pin->PinType.PinCategory == UEdGraphSchema_K2::PC_Class))
 			{
 				if (Pin->PinType.PinSubCategoryObject.IsValid())
 				{
 					PinType = Pin->PinType.PinSubCategoryObject->GetName();
 				}
-				else
-				{
-					PinType = TEXT("");
-				}
+			}
+			else
+			{
+				PinType = Pin->PinType.PinCategory.ToString();
 			}
 
 			if (PinType.IsEmpty())
@@ -715,9 +716,9 @@ bool FKismetCompilerUtilities::ValidateSelfCompatibility(const UEdGraphPin* Pin,
 
 UEdGraphPin* FKismetCompilerUtilities::GenerateAssignmentNodes(class FKismetCompilerContext& CompilerContext, UEdGraph* SourceGraph, UK2Node_CallFunction* CallBeginSpawnNode, UEdGraphNode* SpawnNode, UEdGraphPin* CallBeginResult, const UClass* ForClass )
 {
-	static FString ObjectParamName = FString(TEXT("Object"));
-	static FString ValueParamName = FString(TEXT("Value"));
-	static FString PropertyNameParamName = FString(TEXT("PropertyName"));
+	static const FName ObjectParamName(TEXT("Object"));
+	static const FName ValueParamName(TEXT("Value"));
+	static const FName PropertyNameParamName(TEXT("PropertyName"));
 
 	const UEdGraphSchema_K2* Schema = CompilerContext.GetSchema();
 	UEdGraphPin* LastThen = CallBeginSpawnNode->GetThenPin();
@@ -733,7 +734,7 @@ UEdGraphPin* FKismetCompilerUtilities::GenerateAssignmentNodes(class FKismetComp
 		{
 			if( OrgPin->LinkedTo.Num() == 0 )
 			{
-				UProperty* Property = FindField<UProperty>(ForClass, *OrgPin->PinName);
+				UProperty* Property = FindField<UProperty>(ForClass, OrgPin->PinName);
 				// NULL property indicates that this pin was part of the original node, not the 
 				// class we're assigning to:
 				if( !Property )
@@ -741,12 +742,21 @@ UEdGraphPin* FKismetCompilerUtilities::GenerateAssignmentNodes(class FKismetComp
 					continue;
 				}
 
-				if (ForClass->ClassDefaultObject)
+				// We don't want to generate an assignment node unless the default value 
+				// differs from the value in the CDO:
+				FString DefaultValueAsString;
+					
+				if (FBlueprintCompilationManager::GetDefaultValue(ForClass, Property, DefaultValueAsString))
+				{
+					if (DefaultValueAsString == OrgPin->GetDefaultAsString())
+					{
+						continue;
+					}
+				}
+				else if(ForClass->ClassDefaultObject)
 				{
-					// We don't want to generate an assignment node unless the default value 
-					// differs from the value in the CDO:
-					FString DefaultValueAsString;
 					FBlueprintEditorUtils::PropertyValueToString(Property, (uint8*)ForClass->ClassDefaultObject, DefaultValueAsString);
+
 					if (DefaultValueAsString == OrgPin->GetDefaultAsString())
 					{
 						continue;
@@ -779,12 +789,12 @@ UEdGraphPin* FKismetCompilerUtilities::GenerateAssignmentNodes(class FKismetComp
 
 				// Fill in literal for 'property name' pin - name of pin is property name
 				UEdGraphPin* PropertyNamePin = SetVarNode->FindPinChecked(PropertyNameParamName);
-				PropertyNamePin->DefaultValue = OrgPin->PinName;
+				PropertyNamePin->DefaultValue = OrgPin->PinName.ToString();
 
 				UEdGraphPin* ValuePin = SetVarNode->FindPinChecked(ValueParamName);
 				if (OrgPin->LinkedTo.Num() == 0 &&
 					OrgPin->DefaultValue != FString() &&
-					OrgPin->PinType.PinCategory == Schema->PC_Byte &&
+					OrgPin->PinType.PinCategory == UEdGraphSchema_K2::PC_Byte &&
 					OrgPin->PinType.PinSubCategoryObject.IsValid() &&
 					OrgPin->PinType.PinSubCategoryObject->IsA<UEnum>())
 				{
@@ -792,7 +802,7 @@ UEdGraphPin* FKismetCompilerUtilities::GenerateAssignmentNodes(class FKismetComp
 					UK2Node_EnumLiteral* EnumLiteralNode = CompilerContext.SpawnIntermediateNode<UK2Node_EnumLiteral>(SpawnNode, SourceGraph);
 					EnumLiteralNode->Enum = CastChecked<UEnum>(OrgPin->PinType.PinSubCategoryObject.Get());
 					EnumLiteralNode->AllocateDefaultPins();
-					EnumLiteralNode->FindPinChecked(Schema->PN_ReturnValue)->MakeLinkTo(ValuePin);
+					EnumLiteralNode->FindPinChecked(UEdGraphSchema_K2::PN_ReturnValue)->MakeLinkTo(ValuePin);
 
 					UEdGraphPin* InPin = EnumLiteralNode->FindPinChecked(UK2Node_EnumLiteral::GetEnumInputPinName());
 					check( InPin );
@@ -801,7 +811,7 @@ UEdGraphPin* FKismetCompilerUtilities::GenerateAssignmentNodes(class FKismetComp
 				else
 				{
 					// For non-array struct pins that are not linked, transfer the pin type so that the node will expand an auto-ref that will assign the value by-ref.
-					if (OrgPin->PinType.IsArray() == false && OrgPin->PinType.PinCategory == Schema->PC_Struct && OrgPin->LinkedTo.Num() == 0)
+					if (OrgPin->PinType.IsArray() == false && OrgPin->PinType.PinCategory == UEdGraphSchema_K2::PC_Struct && OrgPin->LinkedTo.Num() == 0)
 					{
 						ValuePin->PinType.PinCategory = OrgPin->PinType.PinCategory;
 						ValuePin->PinType.PinSubCategory = OrgPin->PinType.PinSubCategory;
@@ -838,8 +848,7 @@ void FKismetCompilerUtilities::CreateObjectAssignmentStatement(FKismetFunctionCo
 		ClassTerm->bIsLiteral = true;
 		ClassTerm->Source = DstTerm->Source;
 		ClassTerm->ObjectLiteral = OutputObjClass;
-		check(Context.Schema);
-		ClassTerm->Type.PinCategory = Context.Schema->PC_Class;
+		ClassTerm->Type.PinCategory = UEdGraphSchema_K2::PC_Class;
 
 		EKismetCompiledStatementType CastOpType = bIsOutputInterface ? KCST_CastObjToInterface : KCST_CastInterfaceToObj;
 		FBlueprintCompiledStatement& CastStatement = Context.AppendStatementForNode(Node);
@@ -857,15 +866,14 @@ void FKismetCompilerUtilities::CreateObjectAssignmentStatement(FKismetFunctionCo
 	}
 }
 
-UProperty* FKismetCompilerUtilities::CreatePrimitiveProperty(UObject* PropertyScope, const FName& ValidatedPropertyName, const FString& PinCategory, const FString& PinSubCategory, UObject* PinSubCategoryObject, UClass* SelfClass, bool bIsWeakPointer, const class UEdGraphSchema_K2* Schema, FCompilerResultsLog& MessageLog)
+UProperty* FKismetCompilerUtilities::CreatePrimitiveProperty(UObject* PropertyScope, const FName& ValidatedPropertyName, const FName& PinCategory, const FName& PinSubCategory, UObject* PinSubCategoryObject, UClass* SelfClass, bool bIsWeakPointer, const class UEdGraphSchema_K2* Schema, FCompilerResultsLog& MessageLog)
 {
 	const EObjectFlags ObjectFlags = RF_Public;
 
 	UProperty* NewProperty = nullptr;
-	//@TODO: Nasty string if-else tree
-	if ((PinCategory == Schema->PC_Object) || (PinCategory == Schema->PC_Interface) || (PinCategory == Schema->PC_SoftObject))
+	if ((PinCategory == UEdGraphSchema_K2::PC_Object) || (PinCategory == UEdGraphSchema_K2::PC_Interface) || (PinCategory == UEdGraphSchema_K2::PC_SoftObject))
 	{
-		UClass* SubType = (PinSubCategory == Schema->PSC_Self) ? SelfClass : Cast<UClass>(PinSubCategoryObject);
+		UClass* SubType = (PinSubCategory == UEdGraphSchema_K2::PSC_Self) ? SelfClass : Cast<UClass>(PinSubCategoryObject);
 
 		if (SubType == nullptr)
 		{
@@ -889,9 +897,9 @@ UProperty* FKismetCompilerUtilities::CreatePrimitiveProperty(UObject* PropertySc
 			}
 			else
 			{
-				UObjectPropertyBase* NewPropertyObj = NULL;
+				UObjectPropertyBase* NewPropertyObj = nullptr;
 
-				if (PinCategory == Schema->PC_SoftObject)
+				if (PinCategory == UEdGraphSchema_K2::PC_SoftObject)
 				{
 					NewPropertyObj = NewObject<USoftObjectProperty>(PropertyScope, ValidatedPropertyName, ObjectFlags);
 				}
@@ -912,12 +920,12 @@ UProperty* FKismetCompilerUtilities::CreatePrimitiveProperty(UObject* PropertySc
 			}
 		}
 	}
-	else if (PinCategory == Schema->PC_Struct)
+	else if (PinCategory == UEdGraphSchema_K2::PC_Struct)
 	{
 		if (UScriptStruct* SubType = Cast<UScriptStruct>(PinSubCategoryObject))
 		{
 			FString StructureError;
-			if (FStructureEditorUtils::EStructureError::Ok == FStructureEditorUtils::IsStructureValid(SubType, NULL, &StructureError))
+			if (FStructureEditorUtils::EStructureError::Ok == FStructureEditorUtils::IsStructureValid(SubType, nullptr, &StructureError))
 			{
 				UStructProperty* NewPropertyStruct = NewObject<UStructProperty>(PropertyScope, ValidatedPropertyName, ObjectFlags);
 				NewPropertyStruct->Struct = SubType;
@@ -946,7 +954,7 @@ UProperty* FKismetCompilerUtilities::CreatePrimitiveProperty(UObject* PropertySc
 			}
 		}
 	}
-	else if ((PinCategory == Schema->PC_Class) || (PinCategory == Schema->PC_SoftClass))
+	else if ((PinCategory == UEdGraphSchema_K2::PC_Class) || (PinCategory == UEdGraphSchema_K2::PC_SoftClass))
 	{
 		UClass* SubType = Cast<UClass>(PinSubCategoryObject);
 
@@ -964,7 +972,7 @@ UProperty* FKismetCompilerUtilities::CreatePrimitiveProperty(UObject* PropertySc
 
 		if (SubType)
 		{
-			if (PinCategory == Schema->PC_SoftClass)
+			if (PinCategory == UEdGraphSchema_K2::PC_SoftClass)
 			{
 				USoftClassProperty* SoftClassProperty = NewObject<USoftClassProperty>(PropertyScope, ValidatedPropertyName, ObjectFlags);
 				// we want to use this setter function instead of setting the 
@@ -988,32 +996,32 @@ UProperty* FKismetCompilerUtilities::CreatePrimitiveProperty(UObject* PropertySc
 			}
 		}
 	}
-	else if (PinCategory == Schema->PC_Int)
+	else if (PinCategory == UEdGraphSchema_K2::PC_Int)
 	{
 		NewProperty = NewObject<UIntProperty>(PropertyScope, ValidatedPropertyName, ObjectFlags);
 		NewProperty->SetPropertyFlags(CPF_HasGetValueTypeHash);
 	}
-	else if (PinCategory == Schema->PC_Float)
+	else if (PinCategory == UEdGraphSchema_K2::PC_Float)
 	{
 		NewProperty = NewObject<UFloatProperty>(PropertyScope, ValidatedPropertyName, ObjectFlags);
 		NewProperty->SetPropertyFlags(CPF_HasGetValueTypeHash);
 	}
-	else if (PinCategory == Schema->PC_Boolean)
+	else if (PinCategory == UEdGraphSchema_K2::PC_Boolean)
 	{
 		UBoolProperty* BoolProperty = NewObject<UBoolProperty>(PropertyScope, ValidatedPropertyName, ObjectFlags);
 		BoolProperty->SetBoolSize(sizeof(bool), true);
 		NewProperty = BoolProperty;
 	}
-	else if (PinCategory == Schema->PC_String)
+	else if (PinCategory == UEdGraphSchema_K2::PC_String)
 	{
 		NewProperty = NewObject<UStrProperty>(PropertyScope, ValidatedPropertyName, ObjectFlags);
 		NewProperty->SetPropertyFlags(CPF_HasGetValueTypeHash);
 	}
-	else if (PinCategory == Schema->PC_Text)
+	else if (PinCategory == UEdGraphSchema_K2::PC_Text)
 	{
 		NewProperty = NewObject<UTextProperty>(PropertyScope, ValidatedPropertyName, ObjectFlags);
 	}
-	else if (PinCategory == Schema->PC_Byte)
+	else if (PinCategory == UEdGraphSchema_K2::PC_Byte)
 	{
 		UEnum* Enum = Cast<UEnum>(PinSubCategoryObject);
 
@@ -1037,7 +1045,7 @@ UProperty* FKismetCompilerUtilities::CreatePrimitiveProperty(UObject* PropertySc
 
 		NewProperty->SetPropertyFlags(CPF_HasGetValueTypeHash);
 	}
-	else if (PinCategory == Schema->PC_Name)
+	else if (PinCategory == UEdGraphSchema_K2::PC_Name)
 	{
 		NewProperty = NewObject<UNameProperty>(PropertyScope, ValidatedPropertyName, ObjectFlags);
 		NewProperty->SetPropertyFlags(CPF_HasGetValueTypeHash);
@@ -1116,7 +1124,7 @@ UProperty* FKismetCompilerUtilities::CreatePropertyOnScope(UStruct* Scope, const
 		PropertyScope = Scope;
 	}
 
-	if (Type.PinCategory == Schema->PC_Delegate)
+	if (Type.PinCategory == UEdGraphSchema_K2::PC_Delegate)
 	{
 		if (UFunction* SignatureFunction = FMemberReference::ResolveSimpleMemberReference<UFunction>(Type.PinSubCategoryMemberReference))
 		{
@@ -1125,7 +1133,7 @@ UProperty* FKismetCompilerUtilities::CreatePropertyOnScope(UStruct* Scope, const
 			NewProperty = NewPropertyDelegate;
 		}
 	}
-	else if (Type.PinCategory == Schema->PC_MCDelegate)
+	else if (Type.PinCategory == UEdGraphSchema_K2::PC_MCDelegate)
 	{
 		UFunction* const SignatureFunction = FMemberReference::ResolveSimpleMemberReference<UFunction>(Type.PinSubCategoryMemberReference);
 		UMulticastDelegateProperty* NewPropertyDelegate = NewObject<UMulticastDelegateProperty>(PropertyScope, ValidatedPropertyName, ObjectFlags);
@@ -1277,7 +1285,7 @@ void FKismetCompilerUtilities::ValidateProperEndExecutionPath(FKismetFunctionCon
 					{
 						if (CurrentPin
 							&& (CurrentPin->Direction == EEdGraphPinDirection::EGPD_Output)
-							&& (CurrentPin->PinType.PinCategory == InContext.Schema->PC_Exec))
+							&& (CurrentPin->PinType.PinCategory == UEdGraphSchema_K2::PC_Exec))
 						{
 							if (!CurrentPin->LinkedTo.Num())
 							{
@@ -1321,7 +1329,7 @@ void FKismetCompilerUtilities::ValidateProperEndExecutionPath(FKismetFunctionCon
 					{
 						if (CurrentPin
 							&& (CurrentPin->Direction == EEdGraphPinDirection::EGPD_Input)
-							&& (CurrentPin->PinType.PinCategory == InContext.Schema->PC_Exec)
+							&& (CurrentPin->PinType.PinCategory == UEdGraphSchema_K2::PC_Exec)
 							&& CurrentPin->LinkedTo.Num())
 						{
 							for (auto LinkedPin : CurrentPin->LinkedTo)
@@ -1356,7 +1364,7 @@ void FKismetCompilerUtilities::ValidateProperEndExecutionPath(FKismetFunctionCon
 				{
 					if (CurrentPin
 						&& (CurrentPin->Direction == EEdGraphPinDirection::EGPD_Output)
-						&& (CurrentPin->PinType.PinCategory == InContext.Schema->PC_Exec)
+						&& (CurrentPin->PinType.PinCategory == UEdGraphSchema_K2::PC_Exec)
 						&& CurrentPin->LinkedTo.Num())
 					{
 						auto LinkedPin = CurrentPin->LinkedTo[0];
@@ -1410,7 +1418,7 @@ void FKismetCompilerUtilities::ValidateProperEndExecutionPath(FKismetFunctionCon
 				{
 					if (CurrentPin
 						&& (CurrentPin->Direction == EEdGraphPinDirection::EGPD_Output)
-						&& (CurrentPin->PinType.PinCategory == InContext.Schema->PC_Exec)
+						&& (CurrentPin->PinType.PinCategory == UEdGraphSchema_K2::PC_Exec)
 						&& CurrentPin->LinkedTo.Num())
 					{
 						if (bUnBreakableOutputWasFound)
@@ -1688,7 +1696,7 @@ FString FNetNameMapping::MakeBaseName(const UEdGraphPin* Net)
 	UEdGraphNode* Owner = Net->GetOwningNode();
 	FString Part1 = Owner->GetDescriptiveCompiledName();
 
-	return FString::Printf(TEXT("%s_%s"), *Part1, *Net->PinName);
+	return FString::Printf(TEXT("%s_%s"), *Part1, *Net->PinName.ToString());
 }
 
 FString FNetNameMapping::MakeBaseName(const UEdGraphNode* Net)
@@ -1742,7 +1750,7 @@ FKismetFunctionContext::~FKismetFunctionContext()
 	if (bAllocatedNetNameMap)
 	{
 		delete NetNameMap;
-		NetNameMap = NULL;
+		NetNameMap = nullptr;
 	}
 
 	for (int32 i = 0; i < AllGeneratedStatements.Num(); ++i)
@@ -2190,7 +2198,7 @@ FBPTerminal* FKismetFunctionContext::CreateLocalTerminal(ETerminalSpecification
 	return Result;
 }
 
-FBPTerminal* FKismetFunctionContext::CreateLocalTerminalFromPinAutoChooseScope(UEdGraphPin* Net, const FString& NewName)
+FBPTerminal* FKismetFunctionContext::CreateLocalTerminalFromPinAutoChooseScope(UEdGraphPin* Net, FString NewName)
 {
 	check(Net);
 	bool bSharedTerm = IsEventGraph();
@@ -2208,7 +2216,7 @@ FBPTerminal* FKismetFunctionContext::CreateLocalTerminalFromPinAutoChooseScope(U
 		bSharedTerm = FEventGraphUtils::PinRepresentsSharedTerminal(*Net);
 	}
 	FBPTerminal* Term = new (bSharedTerm ? EventGraphLocals : Locals) FBPTerminal();
-	Term->CopyFromPin(Net, NewName);
+	Term->CopyFromPin(Net, MoveTemp(NewName));
 	return Term;
 }
 
@@ -2217,17 +2225,16 @@ FBPTerminal* FKismetFunctionContext::CreateLocalTerminalFromPinAutoChooseScope(U
 //////////////////////////////////////////////////////////////////////////
 // FBPTerminal
 
-void FBPTerminal::CopyFromPin(UEdGraphPin* Net, const FString& NewName)
+void FBPTerminal::CopyFromPin(UEdGraphPin* Net, FString NewName)
 {
 	Type = Net->PinType;
 	SourcePin = Net;
-	Name = NewName;
+	Name = MoveTemp(NewName);
 
 	bPassedByReference = Net->PinType.bIsReference;
 
-	const UEdGraphSchema_K2* Schema = Cast<const UEdGraphSchema_K2>(Net->GetSchema());
-	const bool bStructCategory = Schema && (Schema->PC_Struct == Net->PinType.PinCategory);
-	const bool bStructSubCategoryObj = (NULL != Cast<UScriptStruct>(Net->PinType.PinSubCategoryObject.Get()));
+	const bool bStructCategory = (UEdGraphSchema_K2::PC_Struct == Net->PinType.PinCategory);
+	const bool bStructSubCategoryObj = (nullptr != Cast<UScriptStruct>(Net->PinType.PinSubCategoryObject.Get()));
 	SetContextTypeStruct(bStructCategory && bStructSubCategoryObj);
 }
 
diff --git a/Engine/Source/Editor/KismetCompiler/Private/KismetCompilerVMBackend.cpp b/Engine/Source/Editor/KismetCompiler/Private/KismetCompilerVMBackend.cpp
--- a/Engine/Source/Editor/KismetCompiler/Private/KismetCompilerVMBackend.cpp
+++ b/Engine/Source/Editor/KismetCompiler/Private/KismetCompilerVMBackend.cpp
@@ -583,7 +583,7 @@ public:
 			// Additional Validation, since we cannot trust custom k2nodes
 			if (CoerceProperty && ensure(Schema) && ensure(CurrentCompilerContext))
 			{
-			    const bool bSecialCaseSelf = (Term->Type.PinSubCategory == Schema->PN_Self);
+			    const bool bSecialCaseSelf = (Term->Type.PinSubCategory == UEdGraphSchema_K2::PN_Self);
 				if(!bSecialCaseSelf)
 			    {
 				    FEdGraphPinType TrueType;
@@ -702,7 +702,7 @@ public:
 				// circumvented with AutoCreateRefTerm, but when it is a self 
 				// (literal) node we still end up here. So, we try to detect and 
 				// handle that case here.
-				if ((Term->Type.PinSubCategory == Schema->PN_Self) && CoerceProperty && CoerceProperty->HasAnyPropertyFlags(CPF_ReferenceParm))
+				if ((Term->Type.PinSubCategory == UEdGraphSchema_K2::PN_Self) && CoerceProperty && CoerceProperty->HasAnyPropertyFlags(CPF_ReferenceParm))
 				{
 					Writer << EX_Self;
 				}
@@ -980,7 +980,7 @@ public:
 			else if (FLiteralTypeHelper::IsObject(&Term->Type, CoerceProperty) || FLiteralTypeHelper::IsClass(&Term->Type, CoerceProperty))
 			{
 				// Note: This case handles both UObjectProperty and UClassProperty
-				if (Term->Type.PinSubCategory == Schema->PN_Self)
+				if (Term->Type.PinSubCategory == UEdGraphSchema_K2::PN_Self)
 				{
 					Writer << EX_Self;
 				}
@@ -996,7 +996,7 @@ public:
 			}
 			else if (FLiteralTypeHelper::IsInterface(&Term->Type, CoerceProperty))
 			{
-				if (Term->Type.PinSubCategory == Schema->PN_Self)
+				if (Term->Type.PinSubCategory == UEdGraphSchema_K2::PN_Self)
 				{
 					Writer << EX_Self;
 				}
@@ -1009,7 +1009,7 @@ public:
 					ensureMsgf(false, TEXT("It is not possible to express this interface property as a literal value! (%s)"), *CoerceProperty->GetFullName());
 				}
 			}
-			else if (!CoerceProperty && Term->Type.PinCategory.IsEmpty() && (Term->Type.PinSubCategory == Schema->PN_Self))
+			else if (!CoerceProperty && Term->Type.PinCategory.IsNone() && (Term->Type.PinSubCategory == UEdGraphSchema_K2::PN_Self))
 			{
 				Writer << EX_Self;
 			}
@@ -1098,7 +1098,7 @@ public:
 			{
 				FBPTerminal CallbackTargetTerm;
 				CallbackTargetTerm.bIsLiteral = true;
-				CallbackTargetTerm.Type.PinSubCategory = Schema->PN_Self;
+				CallbackTargetTerm.Type.PinSubCategory = UEdGraphSchema_K2::PN_Self;
 				EmitTermExpr(&CallbackTargetTerm, Prop);
 			}
 			else
@@ -1347,14 +1347,14 @@ public:
 
 	void EmitDestinationExpression(FBPTerminal* DestinationExpression)
 	{
-		check(Schema && DestinationExpression && !DestinationExpression->Type.PinCategory.IsEmpty());
+		check(Schema && DestinationExpression && !DestinationExpression->Type.PinCategory.IsNone());
 
 		const bool bIsContainer = DestinationExpression->Type.IsContainer();
-		const bool bIsDelegate = Schema->PC_Delegate == DestinationExpression->Type.PinCategory;
-		const bool bIsMulticastDelegate = Schema->PC_MCDelegate == DestinationExpression->Type.PinCategory;
-		const bool bIsBoolean = Schema->PC_Boolean == DestinationExpression->Type.PinCategory;
-		const bool bIsObj = (Schema->PC_Object == DestinationExpression->Type.PinCategory) || (Schema->PC_Class == DestinationExpression->Type.PinCategory);
-		const bool bIsSoftObject = Schema->PC_SoftObject == DestinationExpression->Type.PinCategory;
+		const bool bIsDelegate = UEdGraphSchema_K2::PC_Delegate == DestinationExpression->Type.PinCategory;
+		const bool bIsMulticastDelegate = UEdGraphSchema_K2::PC_MCDelegate == DestinationExpression->Type.PinCategory;
+		const bool bIsBoolean = UEdGraphSchema_K2::PC_Boolean == DestinationExpression->Type.PinCategory;
+		const bool bIsObj = (UEdGraphSchema_K2::PC_Object == DestinationExpression->Type.PinCategory) || (UEdGraphSchema_K2::PC_Class == DestinationExpression->Type.PinCategory);
+		const bool bIsSoftObject = UEdGraphSchema_K2::PC_SoftObject == DestinationExpression->Type.PinCategory;
 		const bool bIsWeakObjPtr = DestinationExpression->Type.bIsWeakPointer;
 
 		if (bIsContainer)
diff --git a/Engine/Source/Editor/KismetCompiler/Private/UserDefinedStructureCompilerUtils.cpp b/Engine/Source/Editor/KismetCompiler/Private/UserDefinedStructureCompilerUtils.cpp
--- a/Engine/Source/Editor/KismetCompiler/Private/UserDefinedStructureCompilerUtils.cpp
+++ b/Engine/Source/Editor/KismetCompiler/Private/UserDefinedStructureCompilerUtils.cpp
@@ -31,7 +31,7 @@ struct FUserDefinedStructureCompilerInner
 		BlueprintsToRecompile.Add(FoundBlueprint, &bAlreadyProcessed);
 		if (!bAlreadyProcessed)
 		{
-			for (auto Function : TFieldRange<UFunction>(FoundBlueprint->GeneratedClass, EFieldIteratorFlags::ExcludeSuper))
+			for (UFunction* Function : TFieldRange<UFunction>(FoundBlueprint->GeneratedClass, EFieldIteratorFlags::ExcludeSuper))
 			{
 				Function->Script.Empty();
 			}
@@ -64,11 +64,11 @@ struct FUserDefinedStructureCompilerInner
 			DuplicatedStruct->AddToRoot();
 			CastChecked<UUserDefinedStructEditorData>(DuplicatedStruct->EditorData)->RecreateDefaultInstance();
 
-			for (auto StructProperty : TObjectRange<UStructProperty>(RF_ClassDefaultObject, /** bIncludeDerivedClasses */ true, /** InternalExcludeFlags */ EInternalObjectFlags::PendingKill))
+			for (UStructProperty* StructProperty : TObjectRange<UStructProperty>(RF_ClassDefaultObject, /** bIncludeDerivedClasses */ true, /** InternalExcludeFlags */ EInternalObjectFlags::PendingKill))
 			{
 				if (StructProperty && (StructureToReinstance == StructProperty->Struct))
 				{
-					if (auto OwnerClass = Cast<UBlueprintGeneratedClass>(StructProperty->GetOwnerClass()))
+					if (UBlueprintGeneratedClass* OwnerClass = Cast<UBlueprintGeneratedClass>(StructProperty->GetOwnerClass()))
 					{
 						if (UBlueprint* FoundBlueprint = Cast<UBlueprint>(OwnerClass->ClassGeneratedBy))
 						{
@@ -76,7 +76,7 @@ struct FUserDefinedStructureCompilerInner
 							ClearStructReferencesInBP(FoundBlueprint, BlueprintsToRecompile);
 						}
 					}
-					else if (auto OwnerStruct = Cast<UUserDefinedStruct>(StructProperty->GetOwnerStruct()))
+					else if (UUserDefinedStruct* OwnerStruct = Cast<UUserDefinedStruct>(StructProperty->GetOwnerStruct()))
 					{
 						check(OwnerStruct != DuplicatedStruct);
 						const bool bValidStruct = (OwnerStruct->GetOutermost() != GetTransientPackage())
@@ -98,7 +98,7 @@ struct FUserDefinedStructureCompilerInner
 
 			DuplicatedStruct->RemoveFromRoot();
 
-			for (auto Blueprint : TObjectRange<UBlueprint>(RF_ClassDefaultObject, /** bIncludeDerivedClasses */ true, /** InternalExcludeFlags */ EInternalObjectFlags::PendingKill))
+			for (UBlueprint* Blueprint : TObjectRange<UBlueprint>(RF_ClassDefaultObject, /** bIncludeDerivedClasses */ true, /** InternalExcludeFlags */ EInternalObjectFlags::PendingKill))
 			{
 				if (Blueprint && !BlueprintsToRecompile.Contains(Blueprint))
 				{
@@ -116,7 +116,7 @@ struct FUserDefinedStructureCompilerInner
 	{
 		check(StructToClean);
 
-		if (auto EditorData = Cast<UUserDefinedStructEditorData>(StructToClean->EditorData))
+		if (UUserDefinedStructEditorData* EditorData = Cast<UUserDefinedStructEditorData>(StructToClean->EditorData))
 		{
 			EditorData->CleanDefaultInstance();
 		}
@@ -132,10 +132,9 @@ struct FUserDefinedStructureCompilerInner
 			TArray<UObject*> SubObjects;
 			GetObjectsWithOuter(StructToClean, SubObjects, true);
 			SubObjects.Remove(StructToClean->EditorData);
-			for (auto SubObjIt = SubObjects.CreateIterator(); SubObjIt; ++SubObjIt)
+			for (UObject* CurrSubObj : SubObjects)
 			{
-				UObject* CurrSubObj = *SubObjIt;
-				CurrSubObj->Rename(NULL, TransientStruct, REN_DontCreateRedirectors);
+				CurrSubObj->Rename(nullptr, TransientStruct, REN_DontCreateRedirectors);
 				if (UProperty* Prop = Cast<UProperty>(CurrSubObj))
 				{
 					FKismetCompilerUtilities::InvalidatePropertyExport(Prop);
@@ -146,15 +145,15 @@ struct FUserDefinedStructureCompilerInner
 				}
 			}
 
-			StructToClean->SetSuperStruct(NULL);
-			StructToClean->Children = NULL;
+			StructToClean->SetSuperStruct(nullptr);
+			StructToClean->Children = nullptr;
 			StructToClean->Script.Empty();
 			StructToClean->MinAlignment = 0;
-			StructToClean->RefLink = NULL;
-			StructToClean->PropertyLink = NULL;
-			StructToClean->DestructorLink = NULL;
+			StructToClean->RefLink = nullptr;
+			StructToClean->PropertyLink = nullptr;
+			StructToClean->DestructorLink = nullptr;
 			StructToClean->ScriptObjectReferences.Empty();
-			StructToClean->PropertyLink = NULL;
+			StructToClean->PropertyLink = nullptr;
 			StructToClean->ErrorMessage.Empty();
 		}
 
@@ -275,7 +274,7 @@ struct FUserDefinedStructureCompilerInner
 
 		Struct->SetMetaData(FBlueprintMetadata::MD_Tooltip, *FStructureEditorUtils::GetTooltip(Struct));
 
-		auto EditorData = CastChecked<UUserDefinedStructEditorData>(Struct->EditorData);
+		UUserDefinedStructEditorData* EditorData = CastChecked<UUserDefinedStructEditorData>(Struct->EditorData);
 		Struct->SetSuperStruct(EditorData->NativeBase);	
 
 		CreateVariables(Struct, K2Schema, MessageLog);
@@ -322,11 +321,10 @@ struct FUserDefinedStructureCompilerInner
 				Struct = ChangedStruct;
 				check(Struct);
 
-				auto Schema = GetDefault<UEdGraphSchema_K2>();
-				for (auto& VarDesc : FStructureEditorUtils::GetVarDesc(Struct))
+				for (FStructVariableDescription& VarDesc : FStructureEditorUtils::GetVarDesc(Struct))
 				{
-					auto StructType = Cast<UUserDefinedStruct>(VarDesc.SubCategoryObject.Get());
-					if (StructType && (VarDesc.Category == Schema->PC_Struct) && AllChangedStructs.Contains(StructType))
+					UUserDefinedStruct* StructType = Cast<UUserDefinedStruct>(VarDesc.SubCategoryObject.Get());
+					if (StructType && (VarDesc.Category == UEdGraphSchema_K2::PC_Struct) && AllChangedStructs.Contains(StructType))
 					{
 						StructuresToWaitFor.Add(StructType);
 					}
@@ -335,10 +333,10 @@ struct FUserDefinedStructureCompilerInner
 		};
 
 		TArray<FDependencyMapEntry> DependencyMap;
-		for (auto Iter = ChangedStructs.CreateConstIterator(); Iter; ++Iter)
+		for (UUserDefinedStruct* ChangedStruct : ChangedStructs)
 		{
 			DependencyMap.Add(FDependencyMapEntry());
-			DependencyMap.Last().Initialize(*Iter, ChangedStructs);
+			DependencyMap.Last().Initialize(ChangedStruct, ChangedStructs);
 		}
 
 		while (DependencyMap.Num())
@@ -361,9 +359,9 @@ struct FUserDefinedStructureCompilerInner
 
 			DependencyMap.RemoveAtSwap(StructureToCompileIndex);
 
-			for (auto EntryIter = DependencyMap.CreateIterator(); EntryIter; ++EntryIter)
+			for (FDependencyMapEntry& MapEntry : DependencyMap)
 			{
-				(*EntryIter).StructuresToWaitFor.Remove(Struct);
+				MapEntry.StructuresToWaitFor.Remove(Struct);
 			}
 		}
 	}
@@ -450,12 +448,12 @@ void FUserDefinedStructureCompilerUtils::CompileStruct(class UUserDefinedStruct*
 			}
 		}
 
-		for (auto BPIter = BlueprintsToRecompile.CreateIterator(); BPIter; ++BPIter)
+		for (UBlueprint* BlueprintToRecompile : BlueprintsToRecompile)
 		{
-			FBlueprintEditorUtils::MarkBlueprintAsStructurallyModified(*BPIter);
+			FBlueprintEditorUtils::MarkBlueprintAsStructurallyModified(BlueprintToRecompile);
 		}
 
-		for (auto ChangedStruct : ChangedStructs)
+		for (UUserDefinedStruct* ChangedStruct : ChangedStructs)
 		{
 			if (ChangedStruct)
 			{
diff --git a/Engine/Source/Editor/KismetCompiler/Public/BPTerminal.h b/Engine/Source/Editor/KismetCompiler/Public/BPTerminal.h
--- a/Engine/Source/Editor/KismetCompiler/Public/BPTerminal.h
+++ b/Engine/Source/Editor/KismetCompiler/Public/BPTerminal.h
@@ -61,7 +61,8 @@ struct FBPTerminal
 	{
 	}
 
-	KISMETCOMPILER_API void CopyFromPin(UEdGraphPin* Net, const FString& NewName);
+	KISMETCOMPILER_API void CopyFromPin(UEdGraphPin* Net, FString NewName);
+	KISMETCOMPILER_API void CopyFromPin(UEdGraphPin* Net, const FName NewName) { CopyFromPin(Net, NewName.ToString()); }
 
 	bool IsTermWritable() const
 	{
diff --git a/Engine/Source/Editor/KismetCompiler/Public/KismetCompiledFunctionContext.h b/Engine/Source/Editor/KismetCompiler/Public/KismetCompiledFunctionContext.h
--- a/Engine/Source/Editor/KismetCompiler/Public/KismetCompiledFunctionContext.h
+++ b/Engine/Source/Editor/KismetCompiler/Public/KismetCompiledFunctionContext.h
@@ -105,7 +105,7 @@ public:
 	bool bUseFlowStack;
 public:
 	FKismetFunctionContext(FCompilerResultsLog& InMessageLog, const UEdGraphSchema_K2* InSchema, UBlueprintGeneratedClass* InNewClass, UBlueprint* InBlueprint, bool bInGeneratingCpp);
-	
+
 	~FKismetFunctionContext();
 
 	void SetExternalNetNameMap(FNetNameMapping* NewMap);
@@ -210,12 +210,12 @@ public:
 	/** Returns a UStruct scope corresponding to the pin type passed in, if one exists */
 	UStruct* GetScopeFromPinType(FEdGraphPinType& Type, UClass* SelfClass)
 	{
-		if ((Type.PinCategory == Schema->PC_Object) || (Type.PinCategory == Schema->PC_Class) || (Type.PinCategory == Schema->PC_Interface))
+		if ((Type.PinCategory == UEdGraphSchema_K2::PC_Object) || (Type.PinCategory == UEdGraphSchema_K2::PC_Class) || (Type.PinCategory == UEdGraphSchema_K2::PC_Interface))
 		{
-			UClass* SubType = (Type.PinSubCategory == Schema->PSC_Self) ? SelfClass : Cast<UClass>(Type.PinSubCategoryObject.Get());
+			UClass* SubType = (Type.PinSubCategory == UEdGraphSchema_K2::PSC_Self) ? SelfClass : Cast<UClass>(Type.PinSubCategoryObject.Get());
 			return SubType;
 		}
-		else if (Type.PinCategory == Schema->PC_Struct)
+		else if (Type.PinCategory == UEdGraphSchema_K2::PC_Struct)
 		{
 			UScriptStruct* SubType = Cast<UScriptStruct>(Type.PinSubCategoryObject.Get());
 			return SubType;
@@ -327,10 +327,10 @@ public:
 	void ResolveStatements();
 
 	/**
-	 * Makes sure an KCST_WireTraceSite is inserted before the specified 
-	 * statement, and associates the specified pin with the inserted wire-trace 
+	 * Makes sure an KCST_WireTraceSite is inserted before the specified
+	 * statement, and associates the specified pin with the inserted wire-trace
 	 * (so we can backwards engineer which pin triggered the goto).
-	 * 
+	 *
 	 * @param  GotoStatement		The statement to insert a goto before.
 	 * @param  AssociatedExecPin	The pin responsible for executing the goto.
 	 */
@@ -362,8 +362,8 @@ public:
 					else if (PrevStatement != NULL)
 					{
 						FBlueprintCompiledStatement* TraceStatement = new FBlueprintCompiledStatement();
-						TraceStatement->Type        = GetWireTraceType();
-						TraceStatement->Comment     = PreJumpNode->NodeComment.IsEmpty() ? PreJumpNode->GetName() : PreJumpNode->NodeComment;
+						TraceStatement->Type = GetWireTraceType();
+						TraceStatement->Comment = PreJumpNode->NodeComment.IsEmpty() ? PreJumpNode->GetName() : PreJumpNode->NodeComment;
 						TraceStatement->ExecContext = AssociatedExecPin;
 
 						NodeStatementList->Insert(TraceStatement, GotoIndex);
@@ -376,12 +376,10 @@ public:
 	}
 
 	/** Looks for a pin of the given name, erroring if the pin is not found or if the direction doesn't match (doesn't verify the pin type) */
-	UEdGraphPin* FindRequiredPinByName(const UEdGraphNode* Node, const FString& PinName, EEdGraphPinDirection RequiredDirection = EGPD_MAX)
+	UEdGraphPin* FindRequiredPinByName(const UEdGraphNode* Node, const FName PinName, EEdGraphPinDirection RequiredDirection = EGPD_MAX)
 	{
-		for (int32 PinIndex = 0; PinIndex < Node->Pins.Num(); ++PinIndex)
+		for (UEdGraphPin* Pin : Node->Pins)
 		{
-			UEdGraphPin* Pin = Node->Pins[PinIndex];
-
 			if (Pin->PinName == PinName)
 			{
 				if ((Pin->Direction == RequiredDirection) || (RequiredDirection == EGPD_MAX))
@@ -391,13 +389,13 @@ public:
 				else
 				{
 					MessageLog.Error(*FString::Printf(TEXT("Expected @@ to be an %s"), (RequiredDirection == EGPD_Output) ? TEXT("output") : TEXT("input")), Pin);
-					return NULL;
+					return nullptr;
 				}
 			}
 		}
 
-		MessageLog.Error(*FString::Printf(TEXT("Expected to find a pin named %s on @@"), *PinName), Node);
-		return NULL;
+		MessageLog.Error(*FString::Printf(TEXT("Expected to find a pin named %s on @@"), *PinName.ToString()), Node);
+		return nullptr;
 	}
 
 	/** Checks to see if a pin is of the requested type */
@@ -422,7 +420,7 @@ public:
 		}
 	}
 
-	KISMETCOMPILER_API FBPTerminal* CreateLocalTerminalFromPinAutoChooseScope(UEdGraphPin* Net, const FString& NewName);
+	KISMETCOMPILER_API FBPTerminal* CreateLocalTerminalFromPinAutoChooseScope(UEdGraphPin* Net, FString NewName);
 	KISMETCOMPILER_API FBPTerminal* CreateLocalTerminal(ETerminalSpecification Spec = ETerminalSpecification::TS_Unspecified);
 };
 
diff --git a/Engine/Source/Editor/KismetCompiler/Public/KismetCompiler.h b/Engine/Source/Editor/KismetCompiler/Public/KismetCompiler.h
--- a/Engine/Source/Editor/KismetCompiler/Public/KismetCompiler.h
+++ b/Engine/Source/Editor/KismetCompiler/Public/KismetCompiler.h
@@ -189,7 +189,7 @@ public:
 			if (SourcePin)
 			{
 				UEdGraphNode* TrueSourceNode = Cast<UEdGraphNode>(MessageLog.FindSourceObject(SourcePin->GetOwningNode()));
-				UEdGraphPin* TrueSourcePin = TrueSourceNode->FindPin(SourcePin->GetName());
+				UEdGraphPin* TrueSourcePin = TrueSourceNode->FindPin(SourcePin->GetFName());
 				SourcePinToExpansionEvent.Add(TrueSourcePin) = Cast<UK2Node_Event>(Result);
 			}
 			else if (SourceNode)
@@ -225,17 +225,39 @@ public:
 	 */
 	FPinConnectionResponse CopyPinLinksToIntermediate(UEdGraphPin& SourcePin, UEdGraphPin& IntermediatePin);
 
-	DEPRECATED(4.17, "Use version that takes PinContainerType instead of separate booleans for array, set, and map")
-	UK2Node_TemporaryVariable* SpawnInternalVariable(UEdGraphNode* SourceNode, FString Category, FString SubCategory, UObject* SubcategoryObject, bool bIsArray, bool bIsSet = false, bool bIsMap = false, const FEdGraphTerminalType& ValueTerminalType = FEdGraphTerminalType());
+	struct FNameParameterHelper
+	{
+		FNameParameterHelper(const FName InNameParameter) : NameParameter(InNameParameter) { }
+		FNameParameterHelper(const FString& InNameParameter) : NameParameter(*InNameParameter) { }
+		FNameParameterHelper(const TCHAR* InNameParameter) : NameParameter(InNameParameter) { }
+
+		FName operator*() const { return NameParameter; }
 
-	UK2Node_TemporaryVariable* SpawnInternalVariable(UEdGraphNode* SourceNode, FString Category, FString SubCategory = FString(), UObject* SubcategoryObject = nullptr, EPinContainerType PinContainerType = EPinContainerType::None, const FEdGraphTerminalType& ValueTerminalType = FEdGraphTerminalType());
+	private:
+		FName NameParameter;
+	};
+
+	DEPRECATED(4.17, "Use version that takes Category and SubCategory as FName, and PinContainerType instead of separate booleans for array, set, and map")
+	UK2Node_TemporaryVariable* SpawnInternalVariable(UEdGraphNode* SourceNode, const FString& Category, const FString& SubCategory, UObject* SubcategoryObject, bool bIsArray, bool bIsSet = false, bool bIsMap = false, const FEdGraphTerminalType& ValueTerminalType = FEdGraphTerminalType());
+
+	DEPRECATED(4.18, "Use version that takes Category and SubCategory as FName")
+	UK2Node_TemporaryVariable* SpawnInternalVariable(UEdGraphNode* SourceNode, const FNameParameterHelper Category, const FString& SubCategory, UObject* SubcategoryObject = nullptr, EPinContainerType PinContainerType = EPinContainerType::None, const FEdGraphTerminalType& ValueTerminalType = FEdGraphTerminalType())
+	{
+		return SpawnInternalVariable(SourceNode, *Category, FName(*SubCategory), SubcategoryObject, PinContainerType, ValueTerminalType);
+	}
+
+	//DEPRECATED(4.18, "Remove when removing versions that take subcategory as FString. Required to avoid ambiguity")
+	UK2Node_TemporaryVariable* SpawnInternalVariable(UEdGraphNode* SourceNode, const FNameParameterHelper Category, const TCHAR* SubCategory, UObject* SubcategoryObject = nullptr, EPinContainerType PinContainerType = EPinContainerType::None, const FEdGraphTerminalType& ValueTerminalType = FEdGraphTerminalType())
+	{
+		return SpawnInternalVariable(SourceNode, *Category, FName(SubCategory), SubcategoryObject, PinContainerType, ValueTerminalType);
+	}
+
+	UK2Node_TemporaryVariable* SpawnInternalVariable(UEdGraphNode* SourceNode, FName Category, FName SubCategory = NAME_None, UObject* SubcategoryObject = nullptr, EPinContainerType PinContainerType = EPinContainerType::None, const FEdGraphTerminalType& ValueTerminalType = FEdGraphTerminalType());
 
 	bool UsePersistentUberGraphFrame() const;
 
 	FString GetGuid(const UEdGraphNode* Node) const;
 
-	static TSharedPtr<FKismetCompilerContext> GetCompilerForBP(UBlueprint* BP, FCompilerResultsLog& InMessageLog, const FKismetCompilerOptions& InCompileOptions);
-	
 	/** Ensures that all variables have valid names for compilation/replication */
 	void ValidateVariableNames();
 	
@@ -308,7 +330,7 @@ protected:
 	virtual bool PinIsImportantForDependancies(const UEdGraphPin* Pin) const override
 	{
 		// The execution wires do not form data dependencies, they are only important for final scheduling and that is handled thru gotos
-		return Pin->PinType.PinCategory != Schema->PC_Exec;
+		return Pin->PinType.PinCategory != UEdGraphSchema_K2::PC_Exec;
 	}
 
 protected:
@@ -399,8 +421,7 @@ protected:
 	 */
 	FName GetUbergraphCallName() const
 	{
-		check(Schema);
-		const FString UbergraphCallString = Schema->FN_ExecuteUbergraphBase.ToString() + TEXT("_") + Blueprint->GetName();
+		const FString UbergraphCallString = UEdGraphSchema_K2::FN_ExecuteUbergraphBase.ToString() + TEXT("_") + Blueprint->GetName();
 		return FName(*UbergraphCallString);
 	}
 
@@ -500,7 +521,7 @@ private:
 	/**
 	 * Handles creating a new event node for a given output on a timeline node utilizing the named function
 	 */
-	void CreatePinEventNodeForTimelineFunction(UK2Node_Timeline* TimelineNode, UEdGraph* SourceGraph, FName FunctionName, const FString& PinName, FName ExecFuncName);
+	void CreatePinEventNodeForTimelineFunction(UK2Node_Timeline* TimelineNode, UEdGraph* SourceGraph, FName FunctionName, const FName PinName, FName ExecFuncName);
 
 	/** Util for creating a node to call a function on a timeline and move connections to it */
 	class UK2Node_CallFunction* CreateCallTimelineFunction(UK2Node_Timeline* TimelineNode, UEdGraph* SourceGraph, FName FunctionName, UEdGraphPin* TimelineVarPin, UEdGraphPin* TimelineFunctionPin);
diff --git a/Engine/Source/Editor/KismetCompiler/Public/KismetCompilerMisc.h b/Engine/Source/Editor/KismetCompiler/Public/KismetCompilerMisc.h
--- a/Engine/Source/Editor/KismetCompiler/Public/KismetCompilerMisc.h
+++ b/Engine/Source/Editor/KismetCompiler/Public/KismetCompilerMisc.h
@@ -70,7 +70,7 @@ public:
 	static const UFunction* FindOverriddenImplementableEvent(const FName& EventName, const UClass* Class);
 
 	/** Helper function for creating property for primitive types. Used only to create inner peroperties for UArrayProperty, USetProperty, and UMapProperty: */
-	static UProperty* CreatePrimitiveProperty( UObject* PropertyScope, const FName& ValidatedPropertyName, const FString& PinCategory, const FString& PinSubCategory, UObject* PinSubCategoryObject, UClass* SelfClass, bool bIsWeakPointer, const class UEdGraphSchema_K2* Schema, FCompilerResultsLog& MessageLog);
+	static UProperty* CreatePrimitiveProperty( UObject* PropertyScope, const FName& ValidatedPropertyName, const FName& PinCategory, const FName& PinSubCategory, UObject* PinSubCategoryObject, UClass* SelfClass, bool bIsWeakPointer, const class UEdGraphSchema_K2* Schema, FCompilerResultsLog& MessageLog);
 
 	/** Creates a property named PropertyName of type PropertyType in the Scope or returns NULL if the type is unknown, but does *not* link that property in */
 	static UProperty* CreatePropertyOnScope(UStruct* Scope, const FName& PropertyName, const FEdGraphPinType& Type, UClass* SelfClass, uint64 PropertyFlags, const class UEdGraphSchema_K2* Schema, FCompilerResultsLog& MessageLog);
diff --git a/Engine/Source/Editor/KismetWidgets/Private/SPinTypeSelector.cpp b/Engine/Source/Editor/KismetWidgets/Private/SPinTypeSelector.cpp
--- a/Engine/Source/Editor/KismetWidgets/Private/SPinTypeSelector.cpp
+++ b/Engine/Source/Editor/KismetWidgets/Private/SPinTypeSelector.cpp
@@ -68,9 +68,9 @@ struct FObjectReferenceType
 	TSharedPtr< SWidget > WidgetToDisplay;
 
 	/** Category that should be used when this item is selected */
-	FString PinCategory;
+	FName PinCategory;
 
-	FObjectReferenceType(FPinTypeTreeItem InPinTypeItem, TSharedRef< SWidget > InWidget, FString InPinCategory)
+	FObjectReferenceType(FPinTypeTreeItem InPinTypeItem, TSharedRef< SWidget > InWidget, FName InPinCategory)
 		: PinTypeItem(InPinTypeItem)
 		, WidgetToDisplay(InWidget)
 		, PinCategory(InPinCategory)
@@ -345,7 +345,7 @@ void SPinTypeSelector::Construct(const FArguments& InArgs, FGetPinTypeTree GetPi
 		)
 		.ContentPadding(0)
 		.ToolTip(IDocumentation::Get()->CreateToolTip( TAttribute<FText>(this, &SPinTypeSelector::GetToolTipForContainerWidget), NULL, *BigTooltipDocLink, TEXT("Containers")))
-		.IsEnabled( TargetPinType.Get().PinCategory != Schema->PC_Exec )
+		.IsEnabled( TargetPinType.Get().PinCategory != UEdGraphSchema_K2::PC_Exec )
 		.Visibility(InArgs._bAllowArrays ? EVisibility::Visible : EVisibility::Collapsed)
 		.ButtonContent()
 		[
@@ -456,11 +456,11 @@ void SPinTypeSelector::Construct(const FArguments& InArgs, FGetPinTypeTree GetPi
 
 FText SPinTypeSelector::GetTypeDescription() const
 {
-	const FString& PinSubCategory = TargetPinType.Get().PinSubCategory;
+	const FName PinSubCategory = TargetPinType.Get().PinSubCategory;
 	const UObject* PinSubCategoryObject = TargetPinType.Get().PinSubCategoryObject.Get();
 	if (PinSubCategory != UEdGraphSchema_K2::PSC_Bitmask && PinSubCategoryObject)
 	{
-		if (auto Field = Cast<const UField>(PinSubCategoryObject))
+		if (const UField* Field = Cast<const UField>(PinSubCategoryObject))
 		{
 			return Field->GetDisplayNameText();
 		}
@@ -474,7 +474,7 @@ FText SPinTypeSelector::GetTypeDescription() const
 
 FText SPinTypeSelector::GetSecondaryTypeDescription() const
 {
-	const FString& PinSubCategory = TargetPinType.Get().PinValueType.TerminalSubCategory;
+	const FName PinSubCategory = TargetPinType.Get().PinValueType.TerminalSubCategory;
 	const UObject* PinSubCategoryObject = TargetPinType.Get().PinValueType.TerminalSubCategoryObject.Get();
 	if (PinSubCategory != UEdGraphSchema_K2::PSC_Bitmask && PinSubCategoryObject)
 	{
@@ -552,7 +552,7 @@ TSharedRef<ITableRow> SPinTypeSelector::GenerateTypeTreeRow(FPinTypeTreeItem InI
 	const FText OrgTooltip = InItem->GetToolTip();
 	const FText Tooltip = !OrgTooltip.IsEmpty() ? OrgTooltip : Description;
 
-	const FString PinTooltipExcerpt = ((PinType.PinCategory != UEdGraphSchema_K2::PC_Byte || PinType.PinSubCategoryObject == nullptr) ? PinType.PinCategory : TEXT("Enum")); 
+	const FString PinTooltipExcerpt = ((PinType.PinCategory != UEdGraphSchema_K2::PC_Byte || PinType.PinSubCategoryObject == nullptr) ? PinType.PinCategory.ToString() : TEXT("Enum")); 
 
 	// If there is a sub-menu for this pin type, we need to bind the function to handle the sub-menu
 	FOnGetContent OnGetContent;
@@ -610,7 +610,7 @@ TSharedRef<ITableRow> SPinTypeSelector::GenerateTypeTreeRow(FPinTypeTreeItem InI
 TSharedRef<SWidget> SPinTypeSelector::CreateObjectReferenceWidget(FPinTypeTreeItem InItem, FEdGraphPinType& InPinType, const FSlateBrush* InIconBrush, FText InSimpleTooltip) const
 {
 	return SNew(SHorizontalBox)
-		.ToolTip(IDocumentation::Get()->CreateToolTip(InSimpleTooltip, NULL, *BigTooltipDocLink, InPinType.PinCategory))
+		.ToolTip(IDocumentation::Get()->CreateToolTip(InSimpleTooltip, nullptr, *BigTooltipDocLink, InPinType.PinCategory.ToString()))
 		+SHorizontalBox::Slot()
 		.AutoWidth()
 		.Padding(1.f)
@@ -757,7 +757,7 @@ TSharedRef< SWidget > SPinTypeSelector::GetAllowedObjectTypes(FPinTypeTreeItem I
 		];
 }
 
-void SPinTypeSelector::OnSelectPinType(FPinTypeTreeItem InItem, FString InPinCategory, bool bForSecondaryType)
+void SPinTypeSelector::OnSelectPinType(FPinTypeTreeItem InItem, FName InPinCategory, bool bForSecondaryType)
 {
 	const FScopedTransaction Transaction( LOCTEXT("ChangeParam", "Change Parameter Type") );
 
@@ -788,11 +788,11 @@ void SPinTypeSelector::OnSelectPinType(FPinTypeTreeItem InItem, FString InPinCat
 		SecondaryTypeComboButton->SetIsOpen(false);
 	}
 
-	if( NewTargetPinType.PinCategory == Schema->PC_Exec )
+	if( NewTargetPinType.PinCategory == UEdGraphSchema_K2::PC_Exec )
 	{
 		NewTargetPinType.ContainerType = EPinContainerType::None;
-		NewTargetPinType.PinValueType.TerminalCategory = FString();
-		NewTargetPinType.PinValueType.TerminalSubCategory = FString();
+		NewTargetPinType.PinValueType.TerminalCategory = NAME_None;
+		NewTargetPinType.PinValueType.TerminalSubCategory = NAME_None;
 		NewTargetPinType.PinValueType.TerminalSubCategoryObject = nullptr;
 	}
 
@@ -843,7 +843,7 @@ void SPinTypeSelector::OnTypeSelectionChanged(FPinTypeTreeItem Selection, ESelec
 			}
 			else
 			{
-				OnSelectPinType(Selection, Selection->GetPossibleObjectReferenceTypes() == static_cast<uint8>(EObjectReferenceType::AllTypes)? Schema->PC_Object : Selection->GetPinType(false).PinCategory, bForSecondaryType);
+				OnSelectPinType(Selection, Selection->GetPossibleObjectReferenceTypes() == static_cast<uint8>(EObjectReferenceType::AllTypes)? UEdGraphSchema_K2::PC_Object : Selection->GetPinType(false).PinCategory, bForSecondaryType);
 			}
 		}
 		else
@@ -1099,7 +1099,7 @@ FText SPinTypeSelector::GetToolTipForArrayWidget() const
 	if(IsEnabled())
 	{
 		// The entire widget may be enabled, but the array button disabled because it is an "exec" pin.
-		if(TargetPinType.Get().PinCategory == Schema->PC_Exec)
+		if(TargetPinType.Get().PinCategory == UEdGraphSchema_K2::PC_Exec)
 		{
 			return LOCTEXT("ArrayCheckBox_ExecDisabled", "Exec pins cannot be arrays.");
 		}
@@ -1111,7 +1111,7 @@ FText SPinTypeSelector::GetToolTipForArrayWidget() const
 
 FText SPinTypeSelector::GetToolTipForContainerWidget() const
 {
-	if (TargetPinType.Get().PinCategory == Schema->PC_Exec)
+	if (TargetPinType.Get().PinCategory == UEdGraphSchema_K2::PC_Exec)
 	{
 		// The entire widget may be enabled, but the container type button may be disabled because it is an "exec" pin.
 		return LOCTEXT("ContainerType_ExecDisabled", "Exec pins cannot be containers.");
diff --git a/Engine/Source/Editor/KismetWidgets/Public/SPinTypeSelector.h b/Engine/Source/Editor/KismetWidgets/Public/SPinTypeSelector.h
--- a/Engine/Source/Editor/KismetWidgets/Public/SPinTypeSelector.h
+++ b/Engine/Source/Editor/KismetWidgets/Public/SPinTypeSelector.h
@@ -200,5 +200,5 @@ protected:
 	 * @param InItem				Item selected
 	 * @param InPinCategory			This is the PinType's category, must be provided separately as the PinType in the tree item is always Object Types for any object related type.
 	 */
-	void OnSelectPinType(FPinTypeTreeItem InItem, FString InPinCategory, bool bForSecondaryType);
+	void OnSelectPinType(FPinTypeTreeItem InItem, FName InPinCategory, bool bForSecondaryType);
 };
diff --git a/Engine/Source/Editor/LevelEditor/Private/LevelViewportActions.cpp b/Engine/Source/Editor/LevelEditor/Private/LevelViewportActions.cpp
--- a/Engine/Source/Editor/LevelEditor/Private/LevelViewportActions.cpp
+++ b/Engine/Source/Editor/LevelEditor/Private/LevelViewportActions.cpp
@@ -152,29 +152,6 @@ void FLevelViewportCommands::RegisterCommands()
 	{
 		UI_COMMAND( ShowAllSprites, "Show All Sprites", "Shows all sprites", EUserInterfaceActionType::Button, FInputChord() );
 		UI_COMMAND( HideAllSprites, "Hide All Sprites", "Hides all sprites", EUserInterfaceActionType::Button, FInputChord() );
-
-		// get all the known layers
-		// Get a fresh list as GUnrealEd->SortedSpriteInfo may not yet be built.
-		TArray<FSpriteCategoryInfo> SortedSpriteInfo;
-		UUnrealEdEngine::MakeSortedSpriteInfo(SortedSpriteInfo);
-
-		FString SpritePrefix = TEXT("ShowSprite_");
-		for( int32 InfoIndex = 0; InfoIndex < SortedSpriteInfo.Num(); ++InfoIndex )
-		{
-			const FSpriteCategoryInfo& SpriteInfo = SortedSpriteInfo[InfoIndex];
-
-			const FName CommandName = FName( *(SpritePrefix + SpriteInfo.Category.ToString()) );
-
-			FFormatNamedArguments Args;
-			Args.Add( TEXT("SpriteName"), SpriteInfo.DisplayName );
-			const FText LocalizedName = FText::Format( NSLOCTEXT("UICommands", "SpriteShowFlagName", "Show {SpriteName} Sprites"), Args );
-
-			TSharedPtr<FUICommandInfo> ShowSpriteCommand 
-				= FUICommandInfoDecl( this->AsShared(), CommandName, LocalizedName, SpriteInfo.Description )
-				.UserInterfaceType( EUserInterfaceActionType::ToggleButton );
-
-			ShowSpriteCommands.Add( FLevelViewportCommands::FShowMenuCommand( ShowSpriteCommand, SpriteInfo.DisplayName ) );
-		}
 	}
 
 	// Generate a command for each Stat category
@@ -358,6 +335,36 @@ void FLevelViewportCommands::RegisterShowVolumeCommands()
 	}
 }
 
+void FLevelViewportCommands::RegisterShowSpriteCommands()
+{
+	// get all the known layers
+	// Get a fresh list as GUnrealEd->SortedSpriteInfo may not yet be built.
+	TArray<FSpriteCategoryInfo> SortedSpriteInfo;
+	UUnrealEdEngine::MakeSortedSpriteInfo(SortedSpriteInfo);
+
+	FString SpritePrefix = TEXT("ShowSprite_");
+	for (int32 InfoIndex = 0; InfoIndex < SortedSpriteInfo.Num(); ++InfoIndex)
+	{
+		const FSpriteCategoryInfo& SpriteInfo = SortedSpriteInfo[InfoIndex];
+		const FName CommandName = FName(*(SpritePrefix + SpriteInfo.Category.ToString()));
+
+		// Only add a command if there is none already for this sprite class.
+		TSharedPtr<FUICommandInfo> FoundSpriteCommand = FInputBindingManager::Get().FindCommandInContext(this->AsShared()->GetContextName(), CommandName);
+		if (!FoundSpriteCommand.IsValid())
+		{
+			FFormatNamedArguments Args;
+			Args.Add(TEXT("SpriteName"), SpriteInfo.DisplayName);
+			const FText LocalizedName = FText::Format(NSLOCTEXT("UICommands", "SpriteShowFlagName", "Show {SpriteName} Sprites"), Args);
+
+			TSharedPtr<FUICommandInfo> ShowSpriteCommand
+				= FUICommandInfoDecl(this->AsShared(), CommandName, LocalizedName, SpriteInfo.Description)
+				.UserInterfaceType(EUserInterfaceActionType::ToggleButton);
+
+			ShowSpriteCommands.Add(FLevelViewportCommands::FShowMenuCommand(ShowSpriteCommand, SpriteInfo.DisplayName));
+		}
+	}
+}
+
 PRAGMA_ENABLE_OPTIMIZATION
 
 #undef LOCTEXT_NAMESPACE
diff --git a/Engine/Source/Editor/LevelEditor/Private/SLevelViewport.cpp b/Engine/Source/Editor/LevelEditor/Private/SLevelViewport.cpp
--- a/Engine/Source/Editor/LevelEditor/Private/SLevelViewport.cpp
+++ b/Engine/Source/Editor/LevelEditor/Private/SLevelViewport.cpp
@@ -1355,8 +1355,10 @@ void SLevelViewport::BindViewCommands( FUICommandList& OutCommandList )
 
 void SLevelViewport::BindShowCommands( FUICommandList& OutCommandList )
 {
+	FLevelViewportCommands& LevelViewportCommands = FLevelViewportCommands::Get();
+
 	OutCommandList.MapAction( 
-		FLevelViewportCommands::Get().UseDefaultShowFlags,
+		LevelViewportCommands.UseDefaultShowFlags,
 		FExecuteAction::CreateSP( this, &SLevelViewport::OnUseDefaultShowFlags, false ) );
 
 	const TArray<FShowFlagData>& ShowFlagData = GetShowFlagMenuItems();
@@ -1368,7 +1370,7 @@ void SLevelViewport::BindShowCommands( FUICommandList& OutCommandList )
 
 		// NOTE: There should be one command per show flag so using ShowFlag as the index to ShowFlagCommands is acceptable
 		OutCommandList.MapAction(
-			FLevelViewportCommands::Get().ShowFlagCommands[ ShowFlag ].ShowMenuItem,
+			LevelViewportCommands.ShowFlagCommands[ ShowFlag ].ShowMenuItem,
 			FExecuteAction::CreateSP( this, &SLevelViewport::ToggleShowFlag, SFData.EngineShowFlagIndex ),
 			FCanExecuteAction(),
 			FIsActionChecked::CreateSP( this, &SLevelViewport::IsShowFlagEnabled, SFData.EngineShowFlagIndex ) );
@@ -1378,25 +1380,23 @@ void SLevelViewport::BindShowCommands( FUICommandList& OutCommandList )
 	{
 		// Map 'Show All' and 'Hide All' commands
 		OutCommandList.MapAction(
-			FLevelViewportCommands::Get().ShowAllVolumes,
+			LevelViewportCommands.ShowAllVolumes,
 			FExecuteAction::CreateSP( this, &SLevelViewport::OnToggleAllVolumeActors, true ) );
 
 		OutCommandList.MapAction(
-			FLevelViewportCommands::Get().HideAllVolumes,
+			LevelViewportCommands.HideAllVolumes,
 			FExecuteAction::CreateSP( this, &SLevelViewport::OnToggleAllVolumeActors, false ) );
 		
+
+		LevelViewportCommands.RegisterShowVolumeCommands();
+		const TArray<FLevelViewportCommands::FShowMenuCommand>& ShowVolumeCommands = LevelViewportCommands.ShowVolumeCommands;
+		for (int32 VolumeCommandIndex = 0; VolumeCommandIndex < ShowVolumeCommands.Num(); ++VolumeCommandIndex)
 		{
-			FLevelViewportCommands& LevelViewportCommands = FLevelViewportCommands::Get();
-			LevelViewportCommands.RegisterShowVolumeCommands();
-			const TArray<FLevelViewportCommands::FShowMenuCommand>& ShowVolumeCommands = LevelViewportCommands.ShowVolumeCommands;
-			for (int32 VolumeCommandIndex = 0; VolumeCommandIndex < ShowVolumeCommands.Num(); ++VolumeCommandIndex)
-			{
-				OutCommandList.MapAction(
-					ShowVolumeCommands[ VolumeCommandIndex ].ShowMenuItem,
-					FExecuteAction::CreateSP( this, &SLevelViewport::ToggleShowVolumeClass, VolumeCommandIndex ),
-					FCanExecuteAction(),
-					FIsActionChecked::CreateSP( this, &SLevelViewport::IsVolumeVisible, VolumeCommandIndex ) );
-			}
+			OutCommandList.MapAction(
+				ShowVolumeCommands[ VolumeCommandIndex ].ShowMenuItem,
+				FExecuteAction::CreateSP( this, &SLevelViewport::ToggleShowVolumeClass, VolumeCommandIndex ),
+				FCanExecuteAction(),
+				FIsActionChecked::CreateSP( this, &SLevelViewport::IsVolumeVisible, VolumeCommandIndex ) );
 		}
 	}
 
@@ -1404,11 +1404,11 @@ void SLevelViewport::BindShowCommands( FUICommandList& OutCommandList )
 	{
 		// Map 'Show All' and 'Hide All' commands
 		OutCommandList.MapAction(
-			FLevelViewportCommands::Get().ShowAllLayers,
+			LevelViewportCommands.ShowAllLayers,
 			FExecuteAction::CreateSP( this, &SLevelViewport::OnToggleAllLayers, true ) );
 
 		OutCommandList.MapAction(
-			FLevelViewportCommands::Get().HideAllLayers,
+			LevelViewportCommands.HideAllLayers,
 			FExecuteAction::CreateSP( this, &SLevelViewport::OnToggleAllLayers, false ) );
 	}
 
@@ -1416,21 +1416,23 @@ void SLevelViewport::BindShowCommands( FUICommandList& OutCommandList )
 	{
 		// Map 'Show All' and 'Hide All' commands
 		OutCommandList.MapAction(
-			FLevelViewportCommands::Get().ShowAllSprites,
+			LevelViewportCommands.ShowAllSprites,
 			FExecuteAction::CreateSP( this, &SLevelViewport::OnToggleAllSpriteCategories, true ) );
 
 		OutCommandList.MapAction(
-			FLevelViewportCommands::Get().HideAllSprites,
+			LevelViewportCommands.HideAllSprites,
 			FExecuteAction::CreateSP( this, &SLevelViewport::OnToggleAllSpriteCategories, false ) );
 
 		// Bind each show flag to the same delegate.  We use the delegate payload system to figure out what show flag we are dealing with
-		for( int32 CategoryIndex = 0; CategoryIndex < GUnrealEd->SpriteIDToIndexMap.Num(); ++CategoryIndex )
+		LevelViewportCommands.RegisterShowSpriteCommands();
+		const TArray<FLevelViewportCommands::FShowMenuCommand>& ShowSpriteCommands = LevelViewportCommands.ShowSpriteCommands;
+		for (int32 SpriteCommandIndex = 0; SpriteCommandIndex < ShowSpriteCommands.Num(); ++SpriteCommandIndex)
 		{
 			OutCommandList.MapAction(
-				FLevelViewportCommands::Get().ShowSpriteCommands[ CategoryIndex ].ShowMenuItem,
-				FExecuteAction::CreateSP( this, &SLevelViewport::ToggleSpriteCategory, CategoryIndex ),
+				ShowSpriteCommands[SpriteCommandIndex].ShowMenuItem,
+				FExecuteAction::CreateSP(this, &SLevelViewport::ToggleSpriteCategory, SpriteCommandIndex),
 				FCanExecuteAction(),
-				FIsActionChecked::CreateSP( this, &SLevelViewport::IsSpriteCategoryVisible, CategoryIndex ) );
+				FIsActionChecked::CreateSP(this, &SLevelViewport::IsSpriteCategoryVisible, SpriteCommandIndex));
 		}
 	}
 
@@ -1438,10 +1440,10 @@ void SLevelViewport::BindShowCommands( FUICommandList& OutCommandList )
 	{
 		// Map 'Hide All' command
 		OutCommandList.MapAction(
-			FLevelViewportCommands::Get().HideAllStats,
+			LevelViewportCommands.HideAllStats,
 			FExecuteAction::CreateSP(this, &SLevelViewport::OnToggleAllStatCommands, false));
 
-		for (auto StatCatIt = FLevelViewportCommands::Get().ShowStatCatCommands.CreateConstIterator(); StatCatIt; ++StatCatIt)
+		for (auto StatCatIt = LevelViewportCommands.ShowStatCatCommands.CreateConstIterator(); StatCatIt; ++StatCatIt)
 		{
 			const TArray< FLevelViewportCommands::FShowMenuCommand >& ShowStatCommands = StatCatIt.Value();
 			for (int32 StatIndex = 0; StatIndex < ShowStatCommands.Num(); ++StatIndex)
diff --git a/Engine/Source/Editor/LevelEditor/Public/LevelViewportActions.h b/Engine/Source/Editor/LevelEditor/Public/LevelViewportActions.h
--- a/Engine/Source/Editor/LevelEditor/Public/LevelViewportActions.h
+++ b/Engine/Source/Editor/LevelEditor/Public/LevelViewportActions.h
@@ -191,6 +191,9 @@ public:
 	/** Registers our commands with the binding system for showing volumes. */
 	void RegisterShowVolumeCommands();
 
+	/** Registers our commands with the binding system for showing sprites. */
+	void RegisterShowSpriteCommands();
+
 private:
 	/** Registers additional commands as they are loaded */
 	void HandleNewStatGroup(const TArray<FStatNameAndInfo>& NameAndInfos);
diff --git a/Engine/Source/Editor/MainFrame/Private/Frame/MainFrameActions.cpp b/Engine/Source/Editor/MainFrame/Private/Frame/MainFrameActions.cpp
--- a/Engine/Source/Editor/MainFrame/Private/Frame/MainFrameActions.cpp
+++ b/Engine/Source/Editor/MainFrame/Private/Frame/MainFrameActions.cpp
@@ -548,7 +548,7 @@ void FMainFrameActionCallbacks::PackageProject( const FName InPlatformInfoName )
 				UnrecoverableError = true;
 			}
 
-			if ((Result & ETargetPlatformReadyStatus::RemoveServerNameEmpty) != 0 && (bProjectHasCode || (!FApp::GetEngineIsPromotedBuild() && !FApp::IsEngineInstalled()) || PackagingSettings->BlueprintNativizationMethod != EProjectPackagingBlueprintNativizationMethod::Disabled))
+			if ((Result & ETargetPlatformReadyStatus::RemoveServerNameEmpty) != 0 && (bProjectHasCode || (!FApp::GetEngineIsPromotedBuild() && !FApp::IsEngineInstalled())))
 			{
 				AddMessageLog(
 					LOCTEXT("RemoveServerNameNotFound", "Remote compiling requires a server name. "),
@@ -653,11 +653,6 @@ void FMainFrameActionCallbacks::PackageProject( const FName InPlatformInfoName )
 		OptionalParams += TEXT(" -nodebuginfo");
 	}
 
-	if (PackagingSettings->BlueprintNativizationMethod != EProjectPackagingBlueprintNativizationMethod::Disabled)
-	{
-		OptionalParams += TEXT(" -nativizeAssets");
-	}
-
 	if (PackagingSettings->bGenerateChunks)
 	{
 		OptionalParams += TEXT(" -manifests");
@@ -688,7 +683,7 @@ void FMainFrameActionCallbacks::PackageProject( const FName InPlatformInfoName )
 	}
 
 	// only build if the project has code that might need to be built
-	if (bProjectHasCode || (!FApp::GetEngineIsPromotedBuild() && !FApp::IsEngineInstalled()) || PackagingSettings->BlueprintNativizationMethod != EProjectPackagingBlueprintNativizationMethod::Disabled)
+	if (bProjectHasCode || (!FApp::GetEngineIsPromotedBuild() && !FApp::IsEngineInstalled()))
 	{
 		OptionalParams += TEXT(" -build");
 	}
diff --git a/Engine/Source/Editor/MaterialEditor/Private/FindInMaterial.cpp b/Engine/Source/Editor/MaterialEditor/Private/FindInMaterial.cpp
--- a/Engine/Source/Editor/MaterialEditor/Private/FindInMaterial.cpp
+++ b/Engine/Source/Editor/MaterialEditor/Private/FindInMaterial.cpp
@@ -322,7 +322,7 @@ void SFindInMaterial::MatchTokens(const TArray<FString> &Tokens)
 			if (Pin && Pin->PinFriendlyName.CompareTo(FText::FromString(TEXT(" "))) != 0)
 			{
 				FText PinName = Pin->GetSchema()->GetPinDisplayName(Pin);
-				FString PinSearchString = Pin->PinName + Pin->PinFriendlyName.ToString() + Pin->DefaultValue + Pin->PinType.PinCategory + Pin->PinType.PinSubCategory + (Pin->PinType.PinSubCategoryObject.IsValid() ? Pin->PinType.PinSubCategoryObject.Get()->GetFullName() : TEXT(""));
+				FString PinSearchString = Pin->PinName.ToString() + Pin->PinFriendlyName.ToString() + Pin->DefaultValue + Pin->PinType.PinCategory.ToString() + Pin->PinType.PinSubCategory.ToString() + (Pin->PinType.PinSubCategoryObject.IsValid() ? Pin->PinType.PinSubCategoryObject.Get()->GetFullName() : TEXT(""));
 				PinSearchString = PinSearchString.Replace(TEXT(" "), TEXT(""));
 				if (StringMatchesSearchTokens(Tokens, PinSearchString))
 				{
