Remove LoadModulesRequiredForCompilation and replace it with an entry in KismetCompiler.build.cs. This ensures that 'modules required for compilation' are not loaded while we are loading other modules

diff --git a/Engine/Source/Editor/Kismet/Private/BlueprintEditorModule.cpp b/Engine/Source/Editor/Kismet/Private/BlueprintEditorModule.cpp
--- a/Engine/Source/Editor/Kismet/Private/BlueprintEditorModule.cpp
+++ b/Engine/Source/Editor/Kismet/Private/BlueprintEditorModule.cpp
@@ -457,9 +457,6 @@ TSharedPtr<FKismetCompilerContext> FBlueprintEditorModule::GetCompiler(UBlueprin
 {
 	if (BP != nullptr)
 	{
-		// Load required modules to make sure we have the required items we have
-		BP->LoadModulesRequiredForCompilation();
-
 		if (BlueprintCompilers.Num() > 0)
 		{
 			for (UClass* Class = BP->GetClass(); Class != UBlueprint::StaticClass(); Class = Class->GetSuperClass())
diff --git a/Engine/Source/Editor/KismetCompiler/KismetCompiler.Build.cs b/Engine/Source/Editor/KismetCompiler/KismetCompiler.Build.cs
--- a/Engine/Source/Editor/KismetCompiler/KismetCompiler.Build.cs
+++ b/Engine/Source/Editor/KismetCompiler/KismetCompiler.Build.cs
@@ -8,6 +8,7 @@ public class KismetCompiler : ModuleRules
 	{
 		PrivateDependencyModuleNames.AddRange(
 			new string[] {
+				"AnimationBlueprintEditor",
 				"Core",
 				"CoreUObject",
 				"Engine",
diff --git a/Engine/Source/Editor/KismetCompiler/Private/KismetCompilerModule.cpp b/Engine/Source/Editor/KismetCompiler/Private/KismetCompilerModule.cpp
--- a/Engine/Source/Editor/KismetCompiler/Private/KismetCompilerModule.cpp
+++ b/Engine/Source/Editor/KismetCompiler/Private/KismetCompilerModule.cpp
@@ -186,9 +186,6 @@ void FKismet2CompilerModule::CompileBlueprint(class UBlueprint* Blueprint, const
 
 	const bool bIsBrandNewBP = (Blueprint->SkeletonGeneratedClass == NULL) && (Blueprint->GeneratedClass == NULL) && (Blueprint->ParentClass != NULL) && !CompileOptions.bIsDuplicationInstigated;
 
-	// load required modules before compiling
-	Blueprint->LoadModulesRequiredForCompilation();
-
 	for ( IBlueprintCompiler* Compiler : Compilers )
 	{
 		Compiler->PreCompile(Blueprint, CompileOptions);
diff --git a/Engine/Source/Editor/UMGEditor/Private/WidgetBlueprint.cpp b/Engine/Source/Editor/UMGEditor/Private/WidgetBlueprint.cpp
--- a/Engine/Source/Editor/UMGEditor/Private/WidgetBlueprint.cpp
+++ b/Engine/Source/Editor/UMGEditor/Private/WidgetBlueprint.cpp
@@ -795,11 +795,4 @@ void UWidgetBlueprint::ForEachSourceWidgetImpl (TFunctionRef<void(UWidget*)> Fn)
 	);
 }
 
-#if WITH_EDITOR
-void UWidgetBlueprint::LoadModulesRequiredForCompilation()
-{
-	static const FName ModuleName(TEXT("UMGEditor"));
-	FModuleManager::Get().LoadModule(ModuleName);
-}
-#endif // WITH_EDITOR
 #undef LOCTEXT_NAMESPACE 
diff --git a/Engine/Source/Editor/UMGEditor/Public/WidgetBlueprint.h b/Engine/Source/Editor/UMGEditor/Public/WidgetBlueprint.h
--- a/Engine/Source/Editor/UMGEditor/Public/WidgetBlueprint.h
+++ b/Engine/Source/Editor/UMGEditor/Public/WidgetBlueprint.h
@@ -260,8 +260,4 @@ public:
 
 private:
 	void ForEachSourceWidgetImpl(TFunctionRef<void(UWidget*)> Fn) const;
-
-#if WITH_EDITOR
-	virtual void LoadModulesRequiredForCompilation() override;
-#endif 
 };
diff --git a/Engine/Source/Runtime/Engine/Classes/Animation/AnimBlueprint.h b/Engine/Source/Runtime/Engine/Classes/Animation/AnimBlueprint.h
--- a/Engine/Source/Runtime/Engine/Classes/Animation/AnimBlueprint.h
+++ b/Engine/Source/Runtime/Engine/Classes/Animation/AnimBlueprint.h
@@ -180,8 +180,4 @@ private:
 	UPROPERTY(duplicatetransient, AssetRegistrySearchable)
 	TSoftObjectPtr<class USkeletalMesh> PreviewSkeletalMesh;
 #endif
-
-#if WITH_EDITOR
-	virtual void LoadModulesRequiredForCompilation() override;
-#endif //WITH_EDITOR
 };
diff --git a/Engine/Source/Runtime/Engine/Classes/Engine/Blueprint.h b/Engine/Source/Runtime/Engine/Classes/Engine/Blueprint.h
--- a/Engine/Source/Runtime/Engine/Classes/Engine/Blueprint.h
+++ b/Engine/Source/Runtime/Engine/Classes/Engine/Blueprint.h
@@ -874,11 +874,6 @@ public:
 	void Message_Note(const FString& MessageToLog);
 	void Message_Warn(const FString& MessageToLog);
 	void Message_Error(const FString& MessageToLog);
-	
-	/** 
-	 * Blueprint can choose to load specific modules for compilation
-	 */
-	virtual void LoadModulesRequiredForCompilation();
 #endif
 	
 };
diff --git a/Engine/Source/Runtime/Engine/Private/AnimBlueprint.cpp b/Engine/Source/Runtime/Engine/Private/AnimBlueprint.cpp
--- a/Engine/Source/Runtime/Engine/Private/AnimBlueprint.cpp
+++ b/Engine/Source/Runtime/Engine/Private/AnimBlueprint.cpp
@@ -202,9 +202,4 @@ bool UAnimBlueprint::CanRecompileWhilePlayingInEditor() const
 	return GetDefault<UEditorExperimentalSettings>()->bEnableLiveRecompilationOfAnimationBlueprints;
 }
 
-void UAnimBlueprint::LoadModulesRequiredForCompilation()
-{
-	static const FName ModuleName(TEXT("AnimationBlueprintEditor"));
-	FModuleManager::Get().LoadModule(ModuleName);
-}
 #endif
diff --git a/Engine/Source/Runtime/Engine/Private/Blueprint.cpp b/Engine/Source/Runtime/Engine/Private/Blueprint.cpp
--- a/Engine/Source/Runtime/Engine/Private/Blueprint.cpp
+++ b/Engine/Source/Runtime/Engine/Private/Blueprint.cpp
@@ -1823,9 +1823,4 @@ UEdGraph* UBlueprint::GetLastEditedUberGraph() const
 	return nullptr;
 }
 
-void UBlueprint::LoadModulesRequiredForCompilation()
-{
-	static const FName ModuleName(TEXT("KismetCompiler"));
-	FModuleManager::Get().LoadModule(ModuleName);
-}
 #endif //WITH_EDITOR
