Moved AnimBlueprintCompilerContext back in to KismetCompiler module, left LoadModulesRequiredForCompilation in place, but changed where it is called (now called by UBlueprint::RegenerateClass - earlier in initialization than it was called before), moved compiler registry into KismetCompiler module

diff --git a/Engine/Source/Editor/AnimationBlueprintEditor/Private/AnimationBlueprintEditorModule.cpp b/Engine/Source/Editor/AnimationBlueprintEditor/Private/AnimationBlueprintEditorModule.cpp
--- a/Engine/Source/Editor/AnimationBlueprintEditor/Private/AnimationBlueprintEditorModule.cpp
+++ b/Engine/Source/Editor/AnimationBlueprintEditor/Private/AnimationBlueprintEditorModule.cpp
@@ -8,7 +8,6 @@
 #include "K2Node_CallFunction.h"
 #include "Animation/AnimBlueprint.h"
 #include "Kismet2/KismetEditorUtilities.h"
-#include "AnimBlueprintCompiler.h"
 
 IMPLEMENT_MODULE( FAnimationBlueprintEditorModule, AnimationBlueprintEditor);
 
@@ -29,23 +28,10 @@ void FAnimationBlueprintEditorModule::StartupModule()
 	FEdGraphUtilities::RegisterVisualPinConnectionFactory(AnimGraphPinConnectionFactory);
 
 	FKismetEditorUtilities::RegisterOnBlueprintCreatedCallback(this, UAnimInstance::StaticClass(), FKismetEditorUtilities::FOnBlueprintCreated::CreateRaw(this, &FAnimationBlueprintEditorModule::OnNewBlueprintCreated));
-
-	// Register widget blueprint compiler we do this no matter what.
-	IKismetCompilerInterface& KismetCompilerModule = FModuleManager::LoadModuleChecked<IKismetCompilerInterface>("KismetCompiler");
-	KismetCompilerModule.GetCompilers().Add(&AnimBlueprintCompiler);
-
-	FBlueprintEditorModule::RegisterCompilerGetter(UAnimBlueprint::StaticClass(), FOnGetBlueprintCompiler::CreateStatic(&FAnimationBlueprintEditorModule::GetCompiler));
 }
 
 void FAnimationBlueprintEditorModule::ShutdownModule()
 {
-	IKismetCompilerInterface* KismetCompilerModule = FModuleManager::GetModulePtr<IKismetCompilerInterface>("KismetCompiler");
-	if (KismetCompilerModule)
-	{
-		KismetCompilerModule->GetCompilers().Remove(&AnimBlueprintCompiler);
-	}
-
-	FBlueprintEditorModule::UnregisterCompilerGetter(UAnimBlueprint::StaticClass());
 	FKismetEditorUtilities::UnregisterAutoBlueprintNodeCreation(this);
 
 	FEdGraphUtilities::UnregisterVisualNodeFactory(AnimGraphNodeFactory);
@@ -95,10 +81,4 @@ void FAnimationBlueprintEditorModule::OnNewBlueprintCreated(UBlueprint* InBluepr
 	}
 }
 
-TSharedPtr<FKismetCompilerContext> FAnimationBlueprintEditorModule::GetCompiler(UBlueprint* BP, FCompilerResultsLog& InMessageLog, const FKismetCompilerOptions& InCompileOptions)
-{
-	return TSharedPtr<FKismetCompilerContext>(new FAnimBlueprintCompilerContext(CastChecked<UAnimBlueprint>(BP), InMessageLog, InCompileOptions, nullptr));
-}
-
-
 #undef LOCTEXT_NAMESPACE
diff --git a/Engine/Source/Editor/AnimationBlueprintEditor/Private/AnimationBlueprintEditorModule.h b/Engine/Source/Editor/AnimationBlueprintEditor/Private/AnimationBlueprintEditorModule.h
--- a/Engine/Source/Editor/AnimationBlueprintEditor/Private/AnimationBlueprintEditorModule.h
+++ b/Engine/Source/Editor/AnimationBlueprintEditor/Private/AnimationBlueprintEditorModule.h
@@ -6,7 +6,6 @@
 #include "Toolkits/IToolkitHost.h"
 #include "Toolkits/AssetEditorToolkit.h"
 #include "IAnimationBlueprintEditorModule.h"
-#include "AnimBlueprintCompiler.h"
 
 struct FAnimationGraphNodeFactory;
 struct FAnimationGraphPinFactory;
@@ -49,9 +48,4 @@ private:
 	TSharedPtr<FAnimationGraphNodeFactory> AnimGraphNodeFactory;
 	TSharedPtr<FAnimationGraphPinFactory> AnimGraphPinFactory;
 	TSharedPtr<FAnimationGraphPinConnectionFactory> AnimGraphPinConnectionFactory;
-
-	static TSharedPtr<FKismetCompilerContext> GetCompiler(UBlueprint* BP, FCompilerResultsLog& InMessageLog, const FKismetCompilerOptions& InCompileOptions);
-
-	/** Compiler customization for animation controllers */
-	FAnimBlueprintCompiler AnimBlueprintCompiler;
 };
diff --git a/Engine/Source/Editor/Kismet/Private/BlueprintCompilationManager.cpp b/Engine/Source/Editor/Kismet/Private/BlueprintCompilationManager.cpp
--- a/Engine/Source/Editor/Kismet/Private/BlueprintCompilationManager.cpp
+++ b/Engine/Source/Editor/Kismet/Private/BlueprintCompilationManager.cpp
@@ -288,7 +288,7 @@ struct FCompilerData
 		InternalOptions.bSaveIntermediateProducts = (UserOptions & EBlueprintCompileOptions::SaveIntermediateProducts ) != EBlueprintCompileOptions::None;
 		InternalOptions.CompileType = bBytecodeOnly ? EKismetCompileType::BytecodeOnly : EKismetCompileType::Full;
 
-		Compiler = FBlueprintEditorModule::GetCompiler(BP, *ActiveResultsLog, InternalOptions);
+		Compiler = FKismetCompilerContext::GetCompilerForBP(BP, *ActiveResultsLog, InternalOptions);
 	}
 
 	bool IsSkeletonOnly() const { return JobType == ECompilationManagerJobType::SkeletonOnly; }
diff --git a/Engine/Source/Editor/Kismet/Private/BlueprintEditorModule.cpp b/Engine/Source/Editor/Kismet/Private/BlueprintEditorModule.cpp
--- a/Engine/Source/Editor/Kismet/Private/BlueprintEditorModule.cpp
+++ b/Engine/Source/Editor/Kismet/Private/BlueprintEditorModule.cpp
@@ -279,8 +279,6 @@ void FBlueprintEditorModule::ShutdownModule()
 	// Unregister internal SCS editor customizations
 	UnregisterSCSEditorCustomization("InstancedStaticMeshComponent");
 
-	BlueprintCompilers.Empty();
-
 	UEdGraphPin::ShutdownVerification();
 }
 
@@ -426,53 +424,4 @@ void FBlueprintEditorModule::PrepareAutoGeneratedDefaultEvents()
 	}
 }
 
-/** Blueprint Compiler map*/
-TMap<TSubclassOf<UBlueprint>, FOnGetBlueprintCompiler> FBlueprintEditorModule::BlueprintCompilers;
-
-void FBlueprintEditorModule::RegisterCompilerGetter(const TSubclassOf<UBlueprint> BlueprintType, FOnGetBlueprintCompiler InCompilerGetter)
-{
-	if (InCompilerGetter.IsBound())
-	{
-		FOnGetBlueprintCompiler* Getter = BlueprintCompilers.Find(BlueprintType);
-		if (Getter)
-		{
-			*Getter = InCompilerGetter;
-		}
-		else
-		{
-			BlueprintCompilers.Add(BlueprintType, InCompilerGetter);
-		}
-	}
-}
-
-void FBlueprintEditorModule::UnregisterCompilerGetter(const TSubclassOf<UBlueprint> BlueprintType)
-{
-	if (UObjectInitialized())
-	{
-		BlueprintCompilers.Remove(BlueprintType);
-	}
-}
-
-TSharedPtr<FKismetCompilerContext> FBlueprintEditorModule::GetCompiler(UBlueprint* BP, FCompilerResultsLog& InMessageLog, const FKismetCompilerOptions& InCompileOptions)
-{
-	if (BP != nullptr)
-	{
-		// Load required modules to make sure we have the required items we have
-		BP->LoadModulesRequiredForCompilation();
-
-		if (BlueprintCompilers.Num() > 0)
-		{
-			for (UClass* Class = BP->GetClass(); Class != UBlueprint::StaticClass(); Class = Class->GetSuperClass())
-			{
-				if (FOnGetBlueprintCompiler* CompilerGetter = BlueprintCompilers.Find(Class))
-				{
-					return CompilerGetter->Execute(BP, InMessageLog, InCompileOptions);
-				}
-			}
-		}
-	}
-	
-	return TSharedPtr<FKismetCompilerContext>(new FKismetCompilerContext(BP, InMessageLog, InCompileOptions, nullptr));
-}
-
 #undef LOCTEXT_NAMESPACE
diff --git a/Engine/Source/Editor/Kismet/Public/BlueprintEditorModule.h b/Engine/Source/Editor/Kismet/Public/BlueprintEditorModule.h
--- a/Engine/Source/Editor/Kismet/Public/BlueprintEditorModule.h
+++ b/Engine/Source/Editor/Kismet/Public/BlueprintEditorModule.h
@@ -90,7 +90,6 @@ public:
 };
 
 DECLARE_DELEGATE_RetVal_OneParam(TSharedRef<class ISCSEditorCustomization>, FSCSEditorCustomizationBuilder, TSharedRef< IBlueprintEditor > /* InBlueprintEditor */);
-DECLARE_DELEGATE_RetVal_ThreeParams(TSharedPtr<FKismetCompilerContext>, FOnGetBlueprintCompiler, UBlueprint*, FCompilerResultsLog&, const FKismetCompilerOptions&);
 
 /**
  * The blueprint editor module provides the blueprint editor application.
@@ -193,24 +192,6 @@ public:
 	DECLARE_EVENT_OneParam( FBlueprintEditorModule, FBlueprintEditorOpenedEvent, EBlueprintType );
 	FBlueprintEditorOpenedEvent& OnBlueprintEditorOpened() { return BlueprintEditorOpened; }
 
-	/**
-	 * Register custom compiler getter for the specific blueprint class (ex. UAnimBlueprint vs UBlueprint)
-	 * @param	UBlueprint class		The Blueprint derived class type 
-	 * @param	InCompilerGetter		The delegate for compiler getter
-	 */
-	KISMET_API static void RegisterCompilerGetter(const TSubclassOf<UBlueprint> BlueprintType, FOnGetBlueprintCompiler InCompilerGetter);
-
-	/**
-	 * Unregister a previously registered compiler getter for the specific blueprint class (ex. UAnimBlueprint vs UBlueprint)
-	 * @param	UBlueprint class		The Blueprint derived class type 
-	 */
-	KISMET_API static void UnregisterCompilerGetter(const TSubclassOf<UBlueprint> BlueprintType);
-
-	/** 
-	 * Return compiler for the specific blueprint instance
-	 */
-	static TSharedPtr<FKismetCompilerContext> GetCompiler(UBlueprint* BP, FCompilerResultsLog& InMessageLog, const FKismetCompilerOptions& InCompileOptions);
-
 	/** 
 	 * Exposes a way for other modules to fold in their own Blueprint editor 
 	 * commands (folded in with other BP editor commands, when the editor is 
@@ -252,7 +233,4 @@ private:
 
 	/** Reference to keep our custom configuration panel alive */
 	TSharedPtr<SWidget> ConfigurationPanel;
-
-	/** list of compiler getters for the blueprint types */
-	static TMap<TSubclassOf<UBlueprint>, FOnGetBlueprintCompiler> BlueprintCompilers;
 };
diff --git a/Engine/Source/Editor/KismetCompiler/Private/AnimBlueprintCompiler.cpp b/Engine/Source/Editor/KismetCompiler/Private/AnimBlueprintCompiler.cpp
new file mode 100644
--- /dev/null
+++ b/Engine/Source/Editor/KismetCompiler/Private/AnimBlueprintCompiler.cpp
@@ -0,0 +1,2999 @@
+// Copyright 1998-2017 Epic Games, Inc. All Rights Reserved.
+
+#include "AnimBlueprintCompiler.h"
+#include "UObject/UObjectHash.h"
+#include "Animation/AnimInstance.h"
+#include "EdGraphUtilities.h"
+#include "K2Node_CallFunction.h"
+#include "K2Node_StructMemberGet.h"
+#include "K2Node_BreakStruct.h"
+#include "K2Node_CallArrayFunction.h"
+#include "K2Node_CustomEvent.h"
+#include "K2Node_Knot.h"
+#include "K2Node_StructMemberSet.h"
+#include "K2Node_VariableGet.h"
+#include "K2Node_VariableSet.h"
+
+#include "AnimationGraphSchema.h"
+#include "K2Node_TransitionRuleGetter.h"
+#include "Kismet/BlueprintFunctionLibrary.h"
+#include "Kismet/KismetArrayLibrary.h"
+#include "Kismet/KismetMathLibrary.h"
+#include "Kismet2/BlueprintEditorUtils.h"
+#include "Kismet2/KismetReinstanceUtilities.h"
+#include "AnimStateNodeBase.h"
+#include "AnimStateNode.h"
+#include "AnimStateConduitNode.h"
+#include "AnimStateEntryNode.h"
+#include "AnimStateTransitionNode.h"
+#include "AnimationCustomTransitionGraph.h"
+#include "AnimationStateGraph.h"
+#include "AnimationStateMachineGraph.h"
+#include "AnimationStateMachineSchema.h"
+#include "AnimationTransitionGraph.h"
+#include "AnimGraphNode_Root.h"
+#include "AnimGraphNode_CustomTransitionResult.h"
+#include "Animation/AnimNode_UseCachedPose.h"
+#include "AnimGraphNode_SaveCachedPose.h"
+#include "AnimGraphNode_UseCachedPose.h"
+#include "AnimGraphNode_StateMachineBase.h"
+#include "AnimGraphNode_StateResult.h"
+#include "AnimGraphNode_SequencePlayer.h"
+#include "AnimGraphNode_TransitionPoseEvaluator.h"
+#include "AnimGraphNode_TransitionResult.h"
+#include "K2Node_AnimGetter.h"
+#include "AnimGraphNode_StateMachine.h"
+#include "Animation/AnimNode_SubInstance.h"
+#include "AnimGraphNode_SubInstance.h"
+#include "AnimGraphNode_Slot.h"
+#include "AnimationEditorUtils.h"
+
+#include "AnimBlueprintPostCompileValidation.h" 
+
+#define LOCTEXT_NAMESPACE "AnimBlueprintCompiler"
+
+//
+// Forward declarations.
+//
+class UAnimStateNodeBase;
+
+//////////////////////////////////////////////////////////////////////////
+// FAnimBlueprintCompilerContext::FEffectiveConstantRecord
+
+bool FAnimBlueprintCompilerContext::FEffectiveConstantRecord::Apply(UObject* Object)
+{
+	uint8* StructPtr = nullptr;
+	uint8* PropertyPtr = nullptr;
+	
+	if(NodeVariableProperty->Struct == FAnimNode_SubInstance::StaticStruct())
+	{
+		PropertyPtr = ConstantProperty->ContainerPtrToValuePtr<uint8>(Object);
+	}
+	else
+	{
+		StructPtr = NodeVariableProperty->ContainerPtrToValuePtr<uint8>(Object);
+		PropertyPtr = ConstantProperty->ContainerPtrToValuePtr<uint8>(StructPtr);
+	}
+
+	if (ArrayIndex != INDEX_NONE)
+	{
+		UArrayProperty* ArrayProperty = CastChecked<UArrayProperty>(ConstantProperty);
+
+		// Peer inside the array
+		FScriptArrayHelper ArrayHelper(ArrayProperty, PropertyPtr);
+
+		if (ArrayHelper.IsValidIndex(ArrayIndex))
+		{
+			FBlueprintEditorUtils::PropertyValueFromString_Direct(ArrayProperty->Inner, LiteralSourcePin->GetDefaultAsString(), ArrayHelper.GetRawPtr(ArrayIndex));
+		}
+		else
+		{
+			return false;
+		}
+	}
+	else
+	{
+		FBlueprintEditorUtils::PropertyValueFromString_Direct(ConstantProperty, LiteralSourcePin->GetDefaultAsString(), PropertyPtr);
+	}
+
+	return true;
+}
+
+//////////////////////////////////////////////////////////////////////////
+// FAnimBlueprintCompiler
+
+FAnimBlueprintCompilerContext::FAnimBlueprintCompilerContext(UAnimBlueprint* SourceSketch, FCompilerResultsLog& InMessageLog, const FKismetCompilerOptions& InCompileOptions, TArray<UObject*>* InObjLoaded)
+	: FKismetCompilerContext(SourceSketch, InMessageLog, InCompileOptions, InObjLoaded)
+	, AnimBlueprint(SourceSketch)
+	, bIsDerivedAnimBlueprint(false)
+{
+	// Make sure the skeleton has finished preloading
+	if (AnimBlueprint->TargetSkeleton != nullptr)
+	{
+		if (FLinkerLoad* Linker = AnimBlueprint->TargetSkeleton->GetLinker())
+		{
+			Linker->Preload(AnimBlueprint->TargetSkeleton);
+		}
+	}
+
+	// Determine if there is an anim blueprint in the ancestry of this class
+	bIsDerivedAnimBlueprint = UAnimBlueprint::FindRootAnimBlueprint(AnimBlueprint) != NULL;
+}
+
+FAnimBlueprintCompilerContext::~FAnimBlueprintCompilerContext()
+{
+}
+
+void FAnimBlueprintCompilerContext::CreateClassVariablesFromBlueprint()
+{
+	FKismetCompilerContext::CreateClassVariablesFromBlueprint();
+
+	if(bGenerateSubInstanceVariables)
+	{	
+		for (UEdGraph* It : Blueprint->UbergraphPages)
+		{
+			TArray<UAnimGraphNode_SubInstance*> SubInstanceNodes;
+			It->GetNodesOfClass(SubInstanceNodes);
+			for( UAnimGraphNode_SubInstance* SubInstance : SubInstanceNodes )
+			{
+				ProcessSubInstance(SubInstance, false);
+			}
+		}
+		
+		if(!bIsDerivedAnimBlueprint)
+		{
+			for (UEdGraph* It : Blueprint->FunctionGraphs)
+			{
+				TArray<UAnimGraphNode_SubInstance*> SubInstanceNodes;
+				It->GetNodesOfClass(SubInstanceNodes);
+				for( UAnimGraphNode_SubInstance* SubInstance : SubInstanceNodes )
+				{
+					ProcessSubInstance(SubInstance, false);
+				}
+			}
+		}
+	}
+}
+
+
+UEdGraphSchema_K2* FAnimBlueprintCompilerContext::CreateSchema()
+{
+	AnimSchema = NewObject<UAnimationGraphSchema>();
+	return AnimSchema;
+}
+
+UK2Node_CallFunction* FAnimBlueprintCompilerContext::SpawnCallAnimInstanceFunction(UEdGraphNode* SourceNode, FName FunctionName)
+{
+	//@TODO: SKELETON: This is a call on a parent function (UAnimInstance::StaticClass() specifically), should we treat it as self or not?
+	UK2Node_CallFunction* FunctionCall = SpawnIntermediateNode<UK2Node_CallFunction>(SourceNode);
+	FunctionCall->FunctionReference.SetSelfMember(FunctionName);
+	FunctionCall->AllocateDefaultPins();
+
+	return FunctionCall;
+}
+
+void FAnimBlueprintCompilerContext::CreateEvaluationHandlerStruct(UAnimGraphNode_Base* VisualAnimNode, FEvaluationHandlerRecord& Record)
+{
+	// Shouldn't create a handler if there is nothing to work with
+	check(Record.ServicedProperties.Num() > 0);
+	check(Record.NodeVariableProperty != NULL);
+
+	// Use the node GUID for a stable name across compiles
+	FString FunctionName = FString::Printf(TEXT("%s_%s_%s_%s"), *Record.EvaluationHandlerProperty->GetName(), *VisualAnimNode->GetOuter()->GetName(), *VisualAnimNode->GetClass()->GetName(), *VisualAnimNode->NodeGuid.ToString());
+	Record.HandlerFunctionName = FName(*FunctionName);
+
+	// check function name isnt already used (data exists that can contain duplicate GUIDs) and apply a numeric extension until it is unique
+	int32 ExtensionIndex = 0;
+	FName* ExistingName = HandlerFunctionNames.Find(Record.HandlerFunctionName);
+	while(ExistingName != nullptr)
+	{
+		FunctionName = FString::Printf(TEXT("%s_%s_%s_%s_%d"), *Record.EvaluationHandlerProperty->GetName(), *VisualAnimNode->GetOuter()->GetName(), *VisualAnimNode->GetClass()->GetName(), *VisualAnimNode->NodeGuid.ToString(), ExtensionIndex);
+		Record.HandlerFunctionName = FName(*FunctionName);
+		ExistingName = HandlerFunctionNames.Find(Record.HandlerFunctionName);
+		ExtensionIndex++;
+	}
+
+	HandlerFunctionNames.Add(Record.HandlerFunctionName);
+	
+	// Add a custom event in the graph
+	UK2Node_CustomEvent* EntryNode = SpawnIntermediateEventNode<UK2Node_CustomEvent>(VisualAnimNode, nullptr, ConsolidatedEventGraph);
+	EntryNode->bInternalEvent = true;
+	EntryNode->CustomFunctionName = Record.HandlerFunctionName;
+	EntryNode->AllocateDefaultPins();
+
+	// The ExecChain is the current exec output pin in the linear chain
+	UEdGraphPin* ExecChain = Schema->FindExecutionPin(*EntryNode, EGPD_Output);
+
+	// Create a struct member write node to store the parameters into the animation node
+	UK2Node_StructMemberSet* AssignmentNode = SpawnIntermediateNode<UK2Node_StructMemberSet>(VisualAnimNode, ConsolidatedEventGraph);
+	AssignmentNode->VariableReference.SetSelfMember(Record.NodeVariableProperty->GetFName());
+	AssignmentNode->StructType = Record.NodeVariableProperty->Struct;
+	AssignmentNode->AllocateDefaultPins();
+
+	// Wire up the variable node execution wires
+	UEdGraphPin* ExecVariablesIn = Schema->FindExecutionPin(*AssignmentNode, EGPD_Input);
+	ExecChain->MakeLinkTo(ExecVariablesIn);
+	ExecChain = Schema->FindExecutionPin(*AssignmentNode, EGPD_Output);
+
+	// Run thru each property
+	TSet<FName> PropertiesBeingSet;
+
+	for (auto TargetPinIt = AssignmentNode->Pins.CreateIterator(); TargetPinIt; ++TargetPinIt)
+	{
+		UEdGraphPin* TargetPin = *TargetPinIt;
+		FName PropertyName(TargetPin->PinName);
+
+		// Does it get serviced by this handler?
+		if (FAnimNodeSinglePropertyHandler* SourceInfo = Record.ServicedProperties.Find(PropertyName))
+		{
+			if (TargetPin->PinType.IsArray())
+			{
+				// Grab the array that we need to set members for
+				UK2Node_StructMemberGet* FetchArrayNode = SpawnIntermediateNode<UK2Node_StructMemberGet>(VisualAnimNode, ConsolidatedEventGraph);
+				FetchArrayNode->VariableReference.SetSelfMember(Record.NodeVariableProperty->GetFName());
+				FetchArrayNode->StructType = Record.NodeVariableProperty->Struct;
+				FetchArrayNode->AllocatePinsForSingleMemberGet(PropertyName);
+
+				UEdGraphPin* ArrayVariableNode = FetchArrayNode->FindPin(PropertyName);
+
+				if (SourceInfo->CopyRecords.Num() > 0)
+				{
+					// Set each element in the array
+					for (FPropertyCopyRecord& CopyRecord : SourceInfo->CopyRecords)
+					{
+						int32 ArrayIndex = CopyRecord.DestArrayIndex;
+						UEdGraphPin* DestPin = CopyRecord.DestPin;
+
+						// Create an array element set node
+						UK2Node_CallArrayFunction* ArrayNode = SpawnIntermediateNode<UK2Node_CallArrayFunction>(VisualAnimNode, ConsolidatedEventGraph);
+						ArrayNode->FunctionReference.SetExternalMember(GET_FUNCTION_NAME_CHECKED(UKismetArrayLibrary, Array_Set), UKismetArrayLibrary::StaticClass());
+						ArrayNode->AllocateDefaultPins();
+
+						// Connect the execution chain
+						ExecChain->MakeLinkTo(ArrayNode->GetExecPin());
+						ExecChain = ArrayNode->GetThenPin();
+
+						// Connect the input array
+						UEdGraphPin* TargetArrayPin = ArrayNode->FindPinChecked(TEXT("TargetArray"));
+						TargetArrayPin->MakeLinkTo(ArrayVariableNode);
+						ArrayNode->PinConnectionListChanged(TargetArrayPin);
+
+						// Set the array index
+						UEdGraphPin* TargetIndexPin = ArrayNode->FindPinChecked(TEXT("Index"));
+						TargetIndexPin->DefaultValue = FString::FromInt(ArrayIndex);
+
+						// Wire up the data input
+						UEdGraphPin* TargetItemPin = ArrayNode->FindPinChecked(TEXT("Item"));
+						TargetItemPin->CopyPersistentDataFromOldPin(*DestPin);
+						MessageLog.NotifyIntermediatePinCreation(TargetItemPin, DestPin);
+						DestPin->BreakAllPinLinks();
+					}
+				}
+			}
+			else
+			{
+				check(!TargetPin->PinType.IsContainer())
+				// Single property
+				if (SourceInfo->CopyRecords.Num() > 0 && SourceInfo->CopyRecords[0].DestPin != nullptr)
+				{
+					UEdGraphPin* DestPin = SourceInfo->CopyRecords[0].DestPin;
+
+					PropertiesBeingSet.Add(DestPin->PinName);
+					TargetPin->CopyPersistentDataFromOldPin(*DestPin);
+					MessageLog.NotifyIntermediatePinCreation(TargetPin, DestPin);
+					DestPin->BreakAllPinLinks();
+				}
+			}
+		}
+	}
+
+	// Remove any unused pins from the assignment node to avoid smashing constant values
+	for (int32 PinIndex = 0; PinIndex < AssignmentNode->ShowPinForProperties.Num(); ++PinIndex)
+	{
+		FOptionalPinFromProperty& TestProperty = AssignmentNode->ShowPinForProperties[PinIndex];
+		TestProperty.bShowPin = PropertiesBeingSet.Contains(TestProperty.PropertyName);
+	}
+	AssignmentNode->ReconstructNode();
+}
+
+void FAnimBlueprintCompilerContext::CreateEvaluationHandlerInstance(UAnimGraphNode_Base* VisualAnimNode, FEvaluationHandlerRecord& Record)
+{
+	// Shouldn't create a handler if there is nothing to work with
+	check(Record.ServicedProperties.Num() > 0);
+	check(Record.NodeVariableProperty != nullptr);
+	check(Record.bServicesInstanceProperties);
+
+	// Use the node GUID for a stable name across compiles
+	FString FunctionName = FString::Printf(TEXT("%s_%s_%s_%s"), *Record.EvaluationHandlerProperty->GetName(), *VisualAnimNode->GetOuter()->GetName(), *VisualAnimNode->GetClass()->GetName(), *VisualAnimNode->NodeGuid.ToString());
+	Record.HandlerFunctionName = FName(*FunctionName);
+
+	// check function name isnt already used (data exists that can contain duplicate GUIDs) and apply a numeric extension until it is unique
+	int32 ExtensionIndex = 0;
+	FName* ExistingName = HandlerFunctionNames.Find(Record.HandlerFunctionName);
+	while(ExistingName != nullptr)
+	{
+		FunctionName = FString::Printf(TEXT("%s_%s_%s_%s_%d"), *Record.EvaluationHandlerProperty->GetName(), *VisualAnimNode->GetOuter()->GetName(), *VisualAnimNode->GetClass()->GetName(), *VisualAnimNode->NodeGuid.ToString(), ExtensionIndex);
+		Record.HandlerFunctionName = FName(*FunctionName);
+		ExistingName = HandlerFunctionNames.Find(Record.HandlerFunctionName);
+		ExtensionIndex++;
+	}
+
+	HandlerFunctionNames.Add(Record.HandlerFunctionName);
+
+	// Add a custom event in the graph
+	UK2Node_CustomEvent* EntryNode = SpawnIntermediateNode<UK2Node_CustomEvent>(VisualAnimNode, ConsolidatedEventGraph);
+	EntryNode->bInternalEvent = true;
+	EntryNode->CustomFunctionName = Record.HandlerFunctionName;
+	EntryNode->AllocateDefaultPins();
+
+	// The ExecChain is the current exec output pin in the linear chain
+	UEdGraphPin* ExecChain = Schema->FindExecutionPin(*EntryNode, EGPD_Output);
+
+	// Need to create a variable set call for each serviced property in the handler
+	for(TPair<FName, FAnimNodeSinglePropertyHandler>& PropHandlerPair : Record.ServicedProperties)
+	{
+		FAnimNodeSinglePropertyHandler& PropHandler = PropHandlerPair.Value;
+		FName PropertyName = PropHandlerPair.Key;
+
+		// Should be true, we only want to deal with instance targets in here
+		check(PropHandler.bInstanceIsTarget);
+
+		for(FPropertyCopyRecord& CopyRecord : PropHandler.CopyRecords)
+		{
+			// New set node for the property
+			UK2Node_VariableSet* VarAssignNode = SpawnIntermediateNode<UK2Node_VariableSet>(VisualAnimNode, ConsolidatedEventGraph);
+			VarAssignNode->VariableReference.SetSelfMember(CopyRecord.DestProperty->GetFName());
+			VarAssignNode->AllocateDefaultPins();
+
+			// Wire up the exec line, and update the end of the chain
+			UEdGraphPin* ExecVariablesIn = Schema->FindExecutionPin(*VarAssignNode, EGPD_Input);
+			ExecChain->MakeLinkTo(ExecVariablesIn);
+			ExecChain = Schema->FindExecutionPin(*VarAssignNode, EGPD_Output);
+
+			// Find the property pin on the set node and configure
+			for(UEdGraphPin* TargetPin : VarAssignNode->Pins)
+			{
+				if(TargetPin->PinType.IsContainer())
+				{
+					// Currently unsupported
+					continue;
+				}
+
+				FName PinPropertyName(TargetPin->PinName);
+
+				if(PinPropertyName == PropertyName)
+				{
+					// This is us, wire up the variable
+					UEdGraphPin* DestPin = CopyRecord.DestPin;
+
+					// Copy the data (link up to the source nodes)
+					TargetPin->CopyPersistentDataFromOldPin(*DestPin);
+					MessageLog.NotifyIntermediatePinCreation(TargetPin, DestPin);
+
+					// Old pin needs to not be connected now - break all its links
+					DestPin->BreakAllPinLinks();
+
+					break;
+				}
+			}
+
+			//VarAssignNode->ReconstructNode();
+		}
+	}
+}
+
+void FAnimBlueprintCompilerContext::ProcessAnimationNode(UAnimGraphNode_Base* VisualAnimNode)
+{
+	// Early out if this node has already been processed
+	if (AllocatedAnimNodes.Contains(VisualAnimNode))
+	{
+		return;
+	}
+
+	// Make sure the visual node has a runtime node template
+	const UScriptStruct* NodeType = VisualAnimNode->GetFNodeType();
+	if (NodeType == NULL)
+	{
+		MessageLog.Error(TEXT("@@ has no animation node member"), VisualAnimNode);
+		return;
+	}
+
+	// Give the visual node a chance to do validation
+	{
+		const int32 PreValidationErrorCount = MessageLog.NumErrors;
+		VisualAnimNode->ValidateAnimNodeDuringCompilation(AnimBlueprint->TargetSkeleton, MessageLog);
+		VisualAnimNode->BakeDataDuringCompilation(MessageLog);
+		if (MessageLog.NumErrors != PreValidationErrorCount)
+		{
+			return;
+		}
+	}
+
+	// Create a property for the node
+	const FString NodeVariableName = ClassScopeNetNameMap.MakeValidName(VisualAnimNode);
+
+	const UAnimationGraphSchema* AnimGraphDefaultSchema = GetDefault<UAnimationGraphSchema>();
+
+	FEdGraphPinType NodeVariableType;
+	NodeVariableType.PinCategory = UAnimationGraphSchema::PC_Struct;
+	NodeVariableType.PinSubCategoryObject = NodeType;
+
+	UStructProperty* NewProperty = Cast<UStructProperty>(CreateVariable(FName(*NodeVariableName), NodeVariableType));
+
+	if (NewProperty == NULL)
+	{
+		MessageLog.Error(TEXT("Failed to create node property for @@"), VisualAnimNode);
+	}
+
+	// Register this node with the compile-time data structures
+	const int32 AllocatedIndex = AllocateNodeIndexCounter++;
+	AllocatedAnimNodes.Add(VisualAnimNode, NewProperty);
+	AllocatedNodePropertiesToNodes.Add(NewProperty, VisualAnimNode);
+	AllocatedAnimNodeIndices.Add(VisualAnimNode, AllocatedIndex);
+	AllocatedPropertiesByIndex.Add(AllocatedIndex, NewProperty);
+
+	UAnimGraphNode_Base* TrueSourceObject = MessageLog.FindSourceObjectTypeChecked<UAnimGraphNode_Base>(VisualAnimNode);
+	SourceNodeToProcessedNodeMap.Add(TrueSourceObject, VisualAnimNode);
+
+	// Register the slightly more permanent debug information
+	NewAnimBlueprintClass->GetAnimBlueprintDebugData().NodePropertyToIndexMap.Add(TrueSourceObject, AllocatedIndex);
+	NewAnimBlueprintClass->GetAnimBlueprintDebugData().NodeGuidToIndexMap.Add(TrueSourceObject->NodeGuid, AllocatedIndex);
+	NewAnimBlueprintClass->GetDebugData().RegisterClassPropertyAssociation(TrueSourceObject, NewProperty);
+
+	// Node-specific compilation that requires compiler state info
+	if (UAnimGraphNode_StateMachineBase* StateMachineInstance = Cast<UAnimGraphNode_StateMachineBase>(VisualAnimNode))
+	{
+		// Compile the state machine
+		ProcessStateMachine(StateMachineInstance);
+	}
+	else if (UAnimGraphNode_UseCachedPose* UseCachedPose = Cast<UAnimGraphNode_UseCachedPose>(VisualAnimNode))
+	{
+		// Handle a save/use cached pose linkage
+		ProcessUseCachedPose(UseCachedPose);
+	}
+	else if(UAnimGraphNode_SubInstance* SubInstanceNode = Cast<UAnimGraphNode_SubInstance>(VisualAnimNode))
+	{
+		ProcessSubInstance(SubInstanceNode, true);
+	}
+
+	// Record pose pins for later patchup and gather pins that have an associated evaluation handler
+	TMap<FName, FEvaluationHandlerRecord> StructEvalHandlers;
+
+	for (auto SourcePinIt = VisualAnimNode->Pins.CreateIterator(); SourcePinIt; ++SourcePinIt)
+	{
+		UEdGraphPin* SourcePin = *SourcePinIt;
+		bool bConsumed = false;
+
+		// Register pose links for future use
+		if ((SourcePin->Direction == EGPD_Input) && (AnimGraphDefaultSchema->IsPosePin(SourcePin->PinType)))
+		{
+			// Input pose pin, going to need to be linked up
+			FPoseLinkMappingRecord LinkRecord = VisualAnimNode->GetLinkIDLocation(NodeType, SourcePin);
+			if (LinkRecord.IsValid())
+			{
+				ValidPoseLinkList.Add(LinkRecord);
+				bConsumed = true;
+			}
+		}
+		else
+		{
+			// The property source for our data, either the struct property for an anim node, or the
+			// owning anim instance if using a sub instance node.
+			UProperty* SourcePinProperty = nullptr;
+			int32 SourceArrayIndex = INDEX_NONE;
+
+			// We have special handling below if we're targeting a subinstance instead of our own instance properties
+			UAnimGraphNode_SubInstance* SubInstanceNode = Cast<UAnimGraphNode_SubInstance>(VisualAnimNode);
+
+			// Does this pin have an associated evaluation handler?
+			if(SubInstanceNode)
+			{
+				// Subinstance nodes use instance properties not node properties as they aren't UObjects
+				// and we can't store non-native properties there
+				SubInstanceNode->GetInstancePinProperty(NewAnimBlueprintClass, SourcePin, SourcePinProperty);
+			}
+			else
+			{
+				VisualAnimNode->GetPinAssociatedProperty(NodeType, SourcePin, /*out*/ SourcePinProperty, /*out*/ SourceArrayIndex);
+			}
+			
+			if (SourcePinProperty != NULL)
+			{
+				if (SourcePin->LinkedTo.Num() == 0)
+				{
+					// Literal that can be pushed into the CDO instead of re-evaluated every frame
+					new (ValidAnimNodePinConstants) FEffectiveConstantRecord(NewProperty, SourcePin, SourcePinProperty, SourceArrayIndex);
+					bConsumed = true;
+				}
+				else
+				{
+					// Dynamic value that needs to be wired up and evaluated each frame
+					const FString& EvaluationHandlerStr = SourcePinProperty->GetMetaData(AnimGraphDefaultSchema->NAME_OnEvaluate);
+					FName EvaluationHandlerName(*EvaluationHandlerStr);
+					if (EvaluationHandlerName == NAME_None)
+					{
+						EvaluationHandlerName = AnimGraphDefaultSchema->DefaultEvaluationHandlerName;
+					}
+
+					FEvaluationHandlerRecord& EvalHandler = StructEvalHandlers.FindOrAdd(EvaluationHandlerName);
+
+					EvalHandler.RegisterPin(SourcePin, SourcePinProperty, SourceArrayIndex);
+
+					if(SubInstanceNode)
+					{
+						EvalHandler.bServicesInstanceProperties = true;
+
+						FAnimNodeSinglePropertyHandler* SinglePropHandler = EvalHandler.ServicedProperties.Find(SourcePinProperty->GetFName());
+						check(SinglePropHandler); // Should have been added in RegisterPin
+
+						// Flag that the target property is actually on the instance class and not the node
+						SinglePropHandler->bInstanceIsTarget = true;
+					}
+
+					bConsumed = true;
+				}
+
+				UEdGraphPin* TrueSourcePin = MessageLog.FindSourcePin(SourcePin);
+				if (TrueSourcePin)
+				{
+					NewAnimBlueprintClass->GetDebugData().RegisterClassPropertyAssociation(TrueSourcePin, SourcePinProperty);
+				}
+			}
+		}
+
+		if (!bConsumed && (SourcePin->Direction == EGPD_Input))
+		{
+			//@TODO: ANIMREFACTOR: It's probably OK to have certain pins ignored eventually, but this is very helpful during development
+			MessageLog.Note(TEXT("@@ was visible but ignored"), SourcePin);
+		}
+	}
+
+	// Match the associated property to each evaluation handler
+	for (TFieldIterator<UProperty> NodePropIt(NodeType); NodePropIt; ++NodePropIt)
+	{
+		if (UStructProperty* StructProp = Cast<UStructProperty>(*NodePropIt))
+		{
+			if (StructProp->Struct == FExposedValueHandler::StaticStruct())
+			{
+				// Register this property to the list of pins that need to be updated
+				// (it's OK if there isn't an entry for this handler, it means that the values are static and don't need to be calculated every frame)
+				FName EvaluationHandlerName(StructProp->GetFName());
+
+				FEvaluationHandlerRecord* pRecord = StructEvalHandlers.Find(EvaluationHandlerName);
+
+				if (pRecord)
+				{
+					pRecord->NodeVariableProperty = NewProperty;
+					pRecord->EvaluationHandlerProperty = StructProp;
+				}
+				
+			}
+		}
+	}
+
+	// Generate a new event to update the value of these properties
+	for (auto HandlerIt = StructEvalHandlers.CreateIterator(); HandlerIt; ++HandlerIt)
+	{
+		FName EvaluationHandlerName = HandlerIt.Key();
+		FEvaluationHandlerRecord& Record = HandlerIt.Value();
+
+		if (Record.IsValid())
+		{
+			// build fast path copy records here
+			// we need to do this at this point as they rely on traversing the original wire path
+			// to determine source data. After we call CreateEvaluationHandlerStruct (etc) the original 
+			// graph is modified to hook up to the evaluation handler custom functions & pins are no longer
+			// available
+			Record.BuildFastPathCopyRecords();
+
+			if(Record.bServicesInstanceProperties)
+			{
+				CreateEvaluationHandlerInstance(VisualAnimNode, Record);
+			}
+			else
+			{
+				CreateEvaluationHandlerStruct(VisualAnimNode, Record);
+			}
+
+			ValidEvaluationHandlerList.Add(Record);
+		}
+		else
+		{
+			MessageLog.Error(*FString::Printf(TEXT("A property on @@ references a non-existent %s property named %s"), *(AnimGraphDefaultSchema->NAME_OnEvaluate.ToString()), *(EvaluationHandlerName.ToString())), VisualAnimNode);
+		}
+	}
+}
+
+void FAnimBlueprintCompilerContext::ProcessUseCachedPose(UAnimGraphNode_UseCachedPose* UseCachedPose)
+{
+	bool bSuccessful = false;
+
+	// if compiling only skeleton, we don't have to worry about linking save node
+	if (CompileOptions.CompileType == EKismetCompileType::SkeletonOnly)
+	{
+		return;
+	}
+
+	// Link to the saved cached pose
+	if(UseCachedPose->SaveCachedPoseNode.IsValid())
+	{
+		if (UAnimGraphNode_SaveCachedPose* AssociatedSaveNode = SaveCachedPoseNodes.FindRef(UseCachedPose->SaveCachedPoseNode->CacheName))
+		{
+			UStructProperty* LinkProperty = FindField<UStructProperty>(FAnimNode_UseCachedPose::StaticStruct(), TEXT("LinkToCachingNode"));
+			check(LinkProperty);
+
+			FPoseLinkMappingRecord LinkRecord = FPoseLinkMappingRecord::MakeFromMember(UseCachedPose, AssociatedSaveNode, LinkProperty);
+			if (LinkRecord.IsValid())
+			{
+				ValidPoseLinkList.Add(LinkRecord);
+			}
+			bSuccessful = true;
+
+			// Save CachePoseName for debug
+			FName CachePoseName = FName(*UseCachedPose->SaveCachedPoseNode->CacheName);
+			UseCachedPose->SaveCachedPoseNode->Node.CachePoseName = CachePoseName;
+			UseCachedPose->Node.CachePoseName = CachePoseName;
+		}
+	}
+	
+	if(!bSuccessful)
+	{
+		MessageLog.Error(*LOCTEXT("NoAssociatedSaveNode", "@@ does not have an associated Save Cached Pose node").ToString(), UseCachedPose);
+	}
+}
+
+void FAnimBlueprintCompilerContext::ProcessSubInstance(UAnimGraphNode_SubInstance* SubInstance, bool bCheckForCycles)
+{
+	if(!SubInstance)
+	{
+		return;
+	}
+
+	for(UEdGraphPin* Pin : SubInstance->Pins)
+	{
+		if (Pin->bOrphanedPin || Pin->PinName == TEXT("InPose") || Pin->PinName == TEXT("Pose"))
+		{
+			continue;
+		}
+
+		// Add prefix to avoid collisions
+		FString PrefixedName = SubInstance->GetPinTargetVariableName(Pin);
+
+		// Create a property on the new class to hold the pin data
+		UProperty* NewProperty = FKismetCompilerUtilities::CreatePropertyOnScope(NewAnimBlueprintClass, FName(*PrefixedName), Pin->PinType, NewAnimBlueprintClass, 0, GetSchema(), MessageLog);
+		if(NewProperty)
+		{
+			FKismetCompilerUtilities::LinkAddedProperty(NewAnimBlueprintClass, NewProperty);
+
+			// Add mappings to the node
+			if(!bGenerateSubInstanceVariables)
+			{
+				FAnimNode_SubInstance& RuntimeNode = SubInstance->Node;
+				
+				if(UProperty* FoundProperty = FindField<UProperty>(*RuntimeNode.InstanceClass, Pin->PinName))
+				{
+					RuntimeNode.SourcePropertyNames.Add(NewProperty->GetFName());
+			
+					// Find the property on the internal instance
+					RuntimeNode.DestPropertyNames.Add(FoundProperty->GetFName());
+				}
+			}
+		}
+	}
+	
+	if(bCheckForCycles)
+	{
+		// Check for duplicated slot and state machine names to warn the user about how these 
+		// are boxed
+		NameToCountMap SlotNameToCountMap;
+		NameToCountMap StateMachineNameToCountMap;
+
+		GetDuplicatedSlotAndStateNames(SubInstance, StateMachineNameToCountMap, SlotNameToCountMap);
+
+
+		for(TPair<FName, int32>& Pair : SlotNameToCountMap)
+		{
+			if(Pair.Value > 1)
+			{
+				// Duplicated slot node
+				FString CompilerMessage = FString::Printf(TEXT("Slot name \"%s\" found across multiple instances. Slots are not visible outside of instances so duplicates or subinstances may not perform as expected."), *Pair.Key.ToString());
+				MessageLog.Warning(*CompilerMessage);
+			}
+		}
+
+		for(TPair<FName, int32>& Pair : StateMachineNameToCountMap)
+		{
+			if(Pair.Value > 1)
+			{
+				// Duplicated slot node
+				FString CompilerMessage = FString::Printf(TEXT("State machine \"%s\" found across multiple instances. States are not visible outside of instances so duplicates or subinstances may not perform as expected."), *Pair.Key.ToString());
+				MessageLog.Warning(*CompilerMessage);
+			}
+		}
+	}
+}
+
+void FAnimBlueprintCompilerContext::GetDuplicatedSlotAndStateNames(UAnimGraphNode_SubInstance* InSubInstance, NameToCountMap& OutStateMachineNameToCountMap, NameToCountMap& OutSlotNameToCountMap)
+{
+	if(!InSubInstance)
+	{
+		// Nothing to inspect
+		return;
+	}
+
+	if(UClass* InstanceClass = *InSubInstance->Node.InstanceClass)
+	{
+		UBlueprint* ClassBP = UBlueprint::GetBlueprintFromClass(InstanceClass);
+
+		TArray<UEdGraph*> AllGraphs;
+		ClassBP->GetAllGraphs(AllGraphs);
+
+		for(UEdGraph* Graph : AllGraphs)
+		{
+			TArray<UAnimGraphNode_StateMachine*> StateMachineNodes;
+			TArray<UAnimGraphNode_Slot*> SlotNodes;
+			TArray<UAnimGraphNode_SubInstance*> SubInstanceNodes;
+
+			Graph->GetNodesOfClass(StateMachineNodes);
+			Graph->GetNodesOfClass(SlotNodes);
+			Graph->GetNodesOfClass(SubInstanceNodes);
+
+			for(UAnimGraphNode_StateMachine* StateMachineNode : StateMachineNodes)
+			{
+				int32& Count = OutStateMachineNameToCountMap.FindOrAdd(FName(*StateMachineNode->GetStateMachineName()));
+				// Add one to count as we've encountered this name
+				Count++;
+			}
+
+			for(UAnimGraphNode_Slot* SlotNode : SlotNodes)
+			{
+				int32& Count = OutSlotNameToCountMap.FindOrAdd(SlotNode->Node.SlotName);
+				Count++;
+			}
+
+			for(UAnimGraphNode_SubInstance* SubInstanceNode : SubInstanceNodes)
+			{
+				GetDuplicatedSlotAndStateNames(SubInstanceNode, OutStateMachineNameToCountMap, OutSlotNameToCountMap);
+			}
+		}
+	}
+}
+
+int32 FAnimBlueprintCompilerContext::GetAllocationIndexOfNode(UAnimGraphNode_Base* VisualAnimNode)
+{
+	ProcessAnimationNode(VisualAnimNode);
+	int32* pResult = AllocatedAnimNodeIndices.Find(VisualAnimNode);
+	return (pResult != NULL) ? *pResult : INDEX_NONE;
+}
+
+void FAnimBlueprintCompilerContext::PruneIsolatedAnimationNodes(const TArray<UAnimGraphNode_Base*>& RootSet, TArray<UAnimGraphNode_Base*>& GraphNodes)
+{
+	struct FNodeVisitorDownPoseWires
+	{
+		TSet<UEdGraphNode*> VisitedNodes;
+		const UAnimationGraphSchema* Schema;
+
+		FNodeVisitorDownPoseWires()
+		{
+			Schema = GetDefault<UAnimationGraphSchema>();
+		}
+
+		void TraverseNodes(UEdGraphNode* Node)
+		{
+			VisitedNodes.Add(Node);
+
+			// Follow every exec output pin
+			for (int32 i = 0; i < Node->Pins.Num(); ++i)
+			{
+				UEdGraphPin* MyPin = Node->Pins[i];
+
+				if ((MyPin->Direction == EGPD_Input) && (Schema->IsPosePin(MyPin->PinType)))
+				{
+					for (int32 j = 0; j < MyPin->LinkedTo.Num(); ++j)
+					{
+						UEdGraphPin* OtherPin = MyPin->LinkedTo[j];
+						UEdGraphNode* OtherNode = OtherPin->GetOwningNode();
+						if (!VisitedNodes.Contains(OtherNode))
+						{
+							TraverseNodes(OtherNode);
+						}
+					}
+				}
+			}
+		}
+	};
+
+	// Prune the nodes that aren't reachable via an animation pose link
+	FNodeVisitorDownPoseWires Visitor;
+
+	for (auto RootIt = RootSet.CreateConstIterator(); RootIt; ++RootIt)
+	{
+		UAnimGraphNode_Base* RootNode = *RootIt;
+		Visitor.TraverseNodes(RootNode);
+	}
+
+	for (int32 NodeIndex = 0; NodeIndex < GraphNodes.Num(); ++NodeIndex)
+	{
+		UAnimGraphNode_Base* Node = GraphNodes[NodeIndex];
+		if (!Visitor.VisitedNodes.Contains(Node) && !IsNodePure(Node))
+		{
+			Node->BreakAllNodeLinks();
+			GraphNodes.RemoveAtSwap(NodeIndex);
+			--NodeIndex;
+		}
+	}
+}
+
+void FAnimBlueprintCompilerContext::ProcessAnimationNodesGivenRoot(TArray<UAnimGraphNode_Base*>& AnimNodeList, const TArray<UAnimGraphNode_Base*>& RootSet)
+{
+	// Now prune based on the root set
+	if (MessageLog.NumErrors == 0)
+	{
+		PruneIsolatedAnimationNodes(RootSet, AnimNodeList);
+	}
+
+	// Process the remaining nodes
+	for (auto SourceNodeIt = AnimNodeList.CreateIterator(); SourceNodeIt; ++SourceNodeIt)
+	{
+		UAnimGraphNode_Base* VisualAnimNode = *SourceNodeIt;
+		ProcessAnimationNode(VisualAnimNode);
+	}
+}
+
+TAutoConsoleVariable<int32> CVarAnimDebugCachePoseNodeUpdateOrder(TEXT("a.Compiler.CachePoseNodeUpdateOrderDebug.Enable"), 0, TEXT("Toggle debugging for CacheNodeUpdateOrder debug during AnimBP compilation"));
+
+void FAnimBlueprintCompilerContext::BuildCachedPoseNodeUpdateOrder()
+{
+	TArray<UAnimGraphNode_Root*> RootNodes;
+	ConsolidatedEventGraph->GetNodesOfClass<UAnimGraphNode_Root>(RootNodes);
+
+	TArray<UAnimGraphNode_SaveCachedPose*> OrderedSavePoseNodes;
+
+	// State results are also "root" nodes, need to find the true root
+	UAnimGraphNode_Root* ActualRootNode = nullptr;
+	for(UAnimGraphNode_Root* PossibleRootNode : RootNodes)
+	{
+		if(PossibleRootNode->GetClass() == UAnimGraphNode_Root::StaticClass())
+		{
+			ActualRootNode = PossibleRootNode;
+			break;
+		}
+	}
+
+	const bool bEnableDebug = (CVarAnimDebugCachePoseNodeUpdateOrder.GetValueOnAnyThread() == 1);
+
+	// Should only have one root node
+	if(ActualRootNode)
+	{
+		TArray<UAnimGraphNode_Base*> VisitedRootNodes;
+	
+		UE_CLOG(bEnableDebug, LogAnimation, Display, TEXT("CachePoseNodeOrdering BEGIN")); 
+
+		CachePoseNodeOrdering_StartNewTraversal(ActualRootNode, OrderedSavePoseNodes, VisitedRootNodes);
+
+		UE_CLOG(bEnableDebug, LogAnimation, Display, TEXT("CachePoseNodeOrdering END"));
+	}
+
+	if (bEnableDebug)
+	{
+		UE_LOG(LogAnimation, Display, TEXT("Ordered Save Pose Node List:"));
+		for (UAnimGraphNode_SaveCachedPose* SavedPoseNode : OrderedSavePoseNodes)
+		{
+			UE_LOG(LogAnimation, Display, TEXT("\t%s"), *SavedPoseNode->Node.CachePoseName.ToString())
+		}
+		UE_LOG(LogAnimation, Display, TEXT("End List"));
+	}
+
+	for(UAnimGraphNode_SaveCachedPose* PoseNode : OrderedSavePoseNodes)
+	{
+		if(int32* NodeIndex = AllocatedAnimNodeIndices.Find(PoseNode))
+		{
+			NewAnimBlueprintClass->OrderedSavedPoseIndices.Add(*NodeIndex);
+		}
+		else
+		{
+			MessageLog.Error(TEXT("Failed to find index for a saved pose node while building ordered pose list."));
+		}
+	}
+}
+
+void FAnimBlueprintCompilerContext::CachePoseNodeOrdering_StartNewTraversal(UAnimGraphNode_Base* InRootNode, TArray<UAnimGraphNode_SaveCachedPose*> &OrderedSavePoseNodes, TArray<UAnimGraphNode_Base*> VisitedRootNodes)
+{
+	check(InRootNode);
+	UAnimGraphNode_SaveCachedPose* RootCacheNode = Cast<UAnimGraphNode_SaveCachedPose>(InRootNode);
+	FString RootName = RootCacheNode ? RootCacheNode->CacheName : InRootNode->GetName();
+
+	const bool bEnableDebug = (CVarAnimDebugCachePoseNodeUpdateOrder.GetValueOnAnyThread() == 1);
+
+	UE_CLOG(bEnableDebug, LogAnimation, Display, TEXT("StartNewTraversal %s"), *RootName);
+
+	// Track which root nodes we've visited to prevent infinite recursion.
+	VisitedRootNodes.Add(InRootNode);
+
+	// Need a list of only what we find here to recurse, we can't do that with the total list
+	TArray<UAnimGraphNode_SaveCachedPose*> InternalOrderedNodes;
+
+	// Traverse whole graph from root collecting SaveCachePose nodes we've touched.
+	CachePoseNodeOrdering_TraverseInternal(InRootNode, InternalOrderedNodes);
+
+	// Process nodes that we've touched 
+	UE_CLOG(bEnableDebug, LogAnimation, Display, TEXT("Process Queue for %s"), *RootName);
+	for (UAnimGraphNode_SaveCachedPose* QueuedCacheNode : InternalOrderedNodes)
+	{
+		if (VisitedRootNodes.Contains(QueuedCacheNode))
+		{
+			UE_CLOG(bEnableDebug, LogAnimation, Display, TEXT("Process Queue SaveCachePose %s. ALREADY VISITED, INFINITE RECURSION DETECTED! SKIPPING"), *QueuedCacheNode->CacheName);
+			MessageLog.Error(*FString::Printf(TEXT("Infinite recursion detected with SaveCachePose %s and %s"), *RootName, *QueuedCacheNode->CacheName));
+			continue;
+		}
+		else
+		{
+			OrderedSavePoseNodes.Remove(QueuedCacheNode);
+			OrderedSavePoseNodes.Add(QueuedCacheNode);
+
+			CachePoseNodeOrdering_StartNewTraversal(QueuedCacheNode, OrderedSavePoseNodes, VisitedRootNodes);
+		}
+	}
+
+	UE_CLOG(bEnableDebug, LogAnimation, Display, TEXT("EndNewTraversal %s"), *RootName);
+}
+
+void FAnimBlueprintCompilerContext::CachePoseNodeOrdering_TraverseInternal(UAnimGraphNode_Base* InAnimGraphNode, TArray<UAnimGraphNode_SaveCachedPose*> &OrderedSavePoseNodes)
+{
+	TArray<UAnimGraphNode_Base*> LinkedAnimNodes;
+	GetLinkedAnimNodes(InAnimGraphNode, LinkedAnimNodes);
+
+	const bool bEnableDebug = (CVarAnimDebugCachePoseNodeUpdateOrder.GetValueOnAnyThread() == 1);
+
+	for (UAnimGraphNode_Base* LinkedNode : LinkedAnimNodes)
+	{
+		UE_CLOG(bEnableDebug, LogAnimation, Display, TEXT("\t Processing %s"), *LinkedNode->GetName());
+		if (UAnimGraphNode_UseCachedPose* UsePoseNode = Cast<UAnimGraphNode_UseCachedPose>(LinkedNode))
+		{
+			if (UAnimGraphNode_SaveCachedPose* SaveNode = UsePoseNode->SaveCachedPoseNode.Get())
+			{
+				UE_CLOG(bEnableDebug, LogAnimation, Display, TEXT("\t Queueing SaveCachePose %s"), *SaveNode->CacheName);
+
+				// Requeue the node we found
+				OrderedSavePoseNodes.Remove(SaveNode);
+				OrderedSavePoseNodes.Add(SaveNode);
+			}
+		}
+		else if (UAnimGraphNode_StateMachine* StateMachineNode = Cast<UAnimGraphNode_StateMachine>(LinkedNode))
+		{
+			for (UEdGraph* StateGraph : StateMachineNode->EditorStateMachineGraph->SubGraphs)
+			{
+				TArray<UAnimGraphNode_StateResult*> ResultNodes;
+				StateGraph->GetNodesOfClass(ResultNodes);
+
+				// We should only get one here but doesn't hurt to loop here in case that changes
+				for (UAnimGraphNode_StateResult* ResultNode : ResultNodes)
+				{
+					CachePoseNodeOrdering_TraverseInternal(ResultNode, OrderedSavePoseNodes);
+				}
+			}
+		}
+		else
+		{
+			CachePoseNodeOrdering_TraverseInternal(LinkedNode, OrderedSavePoseNodes);
+		}
+	}
+}
+
+void FAnimBlueprintCompilerContext::GetLinkedAnimNodes(UAnimGraphNode_Base* InGraphNode, TArray<UAnimGraphNode_Base*> &LinkedAnimNodes)
+{
+	for(UEdGraphPin* Pin : InGraphNode->Pins)
+	{
+		if(Pin->Direction == EEdGraphPinDirection::EGPD_Input &&
+		   Pin->PinType.PinCategory == TEXT("struct"))
+		{
+			if(UScriptStruct* Struct = Cast<UScriptStruct>(Pin->PinType.PinSubCategoryObject.Get()))
+			{
+				if(Struct->IsChildOf(FPoseLinkBase::StaticStruct()))
+				{
+					GetLinkedAnimNodes_TraversePin(Pin, LinkedAnimNodes);
+				}
+			}
+		}
+	}
+}
+
+void FAnimBlueprintCompilerContext::GetLinkedAnimNodes_TraversePin(UEdGraphPin* InPin, TArray<UAnimGraphNode_Base*>& LinkedAnimNodes)
+{
+	if(!InPin)
+	{
+		return;
+	}
+
+	for(UEdGraphPin* LinkedPin : InPin->LinkedTo)
+	{
+		if(!LinkedPin)
+		{
+			continue;
+		}
+		
+		UEdGraphNode* OwningNode = LinkedPin->GetOwningNode();
+
+		if(UK2Node_Knot* InnerKnot = Cast<UK2Node_Knot>(OwningNode))
+		{
+			GetLinkedAnimNodes_TraversePin(InnerKnot->GetInputPin(), LinkedAnimNodes);
+		}
+		else if(UAnimGraphNode_Base* AnimNode = Cast<UAnimGraphNode_Base>(OwningNode))
+		{
+			GetLinkedAnimNodes_ProcessAnimNode(AnimNode, LinkedAnimNodes);
+		}
+	}
+}
+
+void FAnimBlueprintCompilerContext::GetLinkedAnimNodes_ProcessAnimNode(UAnimGraphNode_Base* AnimNode, TArray<UAnimGraphNode_Base *> &LinkedAnimNodes)
+{
+	if(!AllocatedAnimNodes.Contains(AnimNode))
+	{
+		UAnimGraphNode_Base* TrueSourceNode = MessageLog.FindSourceObjectTypeChecked<UAnimGraphNode_Base>(AnimNode);
+
+		if(UAnimGraphNode_Base** AllocatedNode = SourceNodeToProcessedNodeMap.Find(TrueSourceNode))
+		{
+			LinkedAnimNodes.Add(*AllocatedNode);
+		}
+		else
+		{
+			FString ErrorString = FString::Printf(*LOCTEXT("MissingLink", "Missing allocated node for %s while searching for node links - likely due to the node having outstanding errors.").ToString(), *AnimNode->GetName());
+			MessageLog.Error(*ErrorString);
+		}
+	}
+	else
+	{
+		LinkedAnimNodes.Add(AnimNode);
+	}
+}
+
+void FAnimBlueprintCompilerContext::ProcessAllAnimationNodes()
+{
+	// Validate the graph
+	ValidateGraphIsWellFormed(ConsolidatedEventGraph);
+
+	// Validate that we have a skeleton
+	if ((AnimBlueprint->TargetSkeleton == nullptr) && !AnimBlueprint->bIsNewlyCreated)
+	{
+		MessageLog.Error(*LOCTEXT("NoSkeleton", "@@ - The skeleton asset for this animation Blueprint is missing, so it cannot be compiled!").ToString(), AnimBlueprint);
+		return;
+	}
+
+	// Build the raw node list
+	TArray<UAnimGraphNode_Base*> AnimNodeList;
+	ConsolidatedEventGraph->GetNodesOfClass<UAnimGraphNode_Base>(/*out*/ AnimNodeList);
+
+	TArray<UK2Node_TransitionRuleGetter*> Getters;
+	ConsolidatedEventGraph->GetNodesOfClass<UK2Node_TransitionRuleGetter>(/*out*/ Getters);
+
+	// Get anim getters from the root anim graph (processing the nodes below will collect them in nested graphs)
+	TArray<UK2Node_AnimGetter*> RootGraphAnimGetters;
+	ConsolidatedEventGraph->GetNodesOfClass<UK2Node_AnimGetter>(RootGraphAnimGetters);
+
+	// Find the root node
+	UAnimGraphNode_Root* PrePhysicsRoot = NULL;
+	TArray<UAnimGraphNode_Base*> RootSet;
+
+	AllocateNodeIndexCounter = 0;
+	NewAnimBlueprintClass->RootAnimNodeIndex = 0;//INDEX_NONE;
+
+	for (auto SourceNodeIt = AnimNodeList.CreateIterator(); SourceNodeIt; ++SourceNodeIt)
+	{
+		UAnimGraphNode_Base* SourceNode = *SourceNodeIt;
+		UAnimGraphNode_Base* TrueNode = MessageLog.FindSourceObjectTypeChecked<UAnimGraphNode_Base>(SourceNode);
+		TrueNode->BlueprintUsage = EBlueprintUsage::NoProperties;
+
+		if (UAnimGraphNode_Root* PossibleRoot = Cast<UAnimGraphNode_Root>(SourceNode))
+		{
+			if (UAnimGraphNode_Root* Root = ExactCast<UAnimGraphNode_Root>(PossibleRoot))
+			{
+				if (PrePhysicsRoot != NULL)
+				{
+					MessageLog.Error(*FString::Printf(*LOCTEXT("ExpectedOneFunctionEntry_Error", "Expected only one animation root, but found both @@ and @@").ToString()),
+						PrePhysicsRoot, Root);
+				}
+				else
+				{
+					RootSet.Add(Root);
+					PrePhysicsRoot = Root;
+				}
+			}
+		}
+		else if (UAnimGraphNode_SaveCachedPose* SavePoseRoot = Cast<UAnimGraphNode_SaveCachedPose>(SourceNode))
+		{
+			//@TODO: Ideally we only add these if there is a UseCachedPose node referencing them, but those can be anywhere and are hard to grab
+			SaveCachedPoseNodes.Add(SavePoseRoot->CacheName, SavePoseRoot);
+			RootSet.Add(SavePoseRoot);
+		}
+	}
+
+	if (PrePhysicsRoot != NULL)
+	{
+		// Process the animation nodes
+		ProcessAnimationNodesGivenRoot(AnimNodeList, RootSet);
+
+		// Process the getter nodes in the graph if there were any
+		for (auto GetterIt = Getters.CreateIterator(); GetterIt; ++GetterIt)
+		{
+			ProcessTransitionGetter(*GetterIt, NULL); // transition nodes should not appear at top-level
+		}
+
+		// Wire root getters
+		for(UK2Node_AnimGetter* RootGraphGetter : RootGraphAnimGetters)
+		{
+			AutoWireAnimGetter(RootGraphGetter, nullptr);
+		}
+
+		// Wire nested getters
+		for(UK2Node_AnimGetter* Getter : FoundGetterNodes)
+		{
+			AutoWireAnimGetter(Getter, nullptr);
+		}
+
+		NewAnimBlueprintClass->RootAnimNodeIndex = GetAllocationIndexOfNode(PrePhysicsRoot);
+	}
+	else
+	{
+		MessageLog.Error(*FString::Printf(*LOCTEXT("ExpectedAFunctionEntry_Error", "Expected an animation root, but did not find one").ToString()));
+	}
+
+	if(CompileOptions.CompileType != EKismetCompileType::SkeletonOnly)
+	{
+		// Build cached pose map
+		BuildCachedPoseNodeUpdateOrder();
+	}
+}
+
+int32 FAnimBlueprintCompilerContext::ExpandGraphAndProcessNodes(UEdGraph* SourceGraph, UAnimGraphNode_Base* SourceRootNode, UAnimStateTransitionNode* TransitionNode, TArray<UEdGraphNode*>* ClonedNodes)
+{
+	// Clone the nodes from the source graph
+	UEdGraph* ClonedGraph = FEdGraphUtilities::CloneGraph(SourceGraph, NULL, &MessageLog, true);
+
+	// Grab all the animation nodes and find the corresponding root node in the cloned set
+	UAnimGraphNode_Base* TargetRootNode = NULL;
+	TArray<UAnimGraphNode_Base*> AnimNodeList;
+	TArray<UK2Node_TransitionRuleGetter*> Getters;
+	TArray<UK2Node_AnimGetter*> AnimGetterNodes;
+
+	for (auto NodeIt = ClonedGraph->Nodes.CreateIterator(); NodeIt; ++NodeIt)
+	{
+		UEdGraphNode* Node = *NodeIt;
+
+		if (UK2Node_TransitionRuleGetter* GetterNode = Cast<UK2Node_TransitionRuleGetter>(Node))
+		{
+			Getters.Add(GetterNode);
+		}
+		else if(UK2Node_AnimGetter* NewGetterNode = Cast<UK2Node_AnimGetter>(Node))
+		{
+			AnimGetterNodes.Add(NewGetterNode);
+		}
+		else if (UAnimGraphNode_Base* TestNode = Cast<UAnimGraphNode_Base>(Node))
+		{
+			AnimNodeList.Add(TestNode);
+
+			//@TODO: There ought to be a better way to determine this
+			if (MessageLog.FindSourceObject(TestNode) == MessageLog.FindSourceObject(SourceRootNode))
+			{
+				TargetRootNode = TestNode;
+			}
+		}
+
+		if (ClonedNodes != NULL)
+		{
+			ClonedNodes->Add(Node);
+		}
+	}
+	check(TargetRootNode);
+
+	// Move the cloned nodes into the consolidated event graph
+	const bool bIsLoading = Blueprint->bIsRegeneratingOnLoad || IsAsyncLoading();
+	const bool bIsCompiling = Blueprint->bBeingCompiled;
+	ClonedGraph->MoveNodesToAnotherGraph(ConsolidatedEventGraph, bIsLoading, bIsCompiling);
+
+	// Process any animation nodes
+	{
+		TArray<UAnimGraphNode_Base*> RootSet;
+		RootSet.Add(TargetRootNode);
+		ProcessAnimationNodesGivenRoot(AnimNodeList, RootSet);
+	}
+
+	// Process the getter nodes in the graph if there were any
+	for (auto GetterIt = Getters.CreateIterator(); GetterIt; ++GetterIt)
+	{
+		ProcessTransitionGetter(*GetterIt, TransitionNode);
+	}
+
+	// Wire anim getter nodes
+	for(UK2Node_AnimGetter* GetterNode : AnimGetterNodes)
+	{
+		FoundGetterNodes.Add(GetterNode);
+	}
+
+	// Returns the index of the processed cloned version of SourceRootNode
+	return GetAllocationIndexOfNode(TargetRootNode);	
+}
+
+void FAnimBlueprintCompilerContext::ProcessStateMachine(UAnimGraphNode_StateMachineBase* StateMachineInstance)
+{
+	struct FMachineCreator
+	{
+	public:
+		int32 MachineIndex;
+		TMap<UAnimStateNodeBase*, int32> StateIndexTable;
+		TMap<UAnimStateTransitionNode*, int32> TransitionIndexTable;
+		UAnimBlueprintGeneratedClass* AnimBlueprintClass;
+		UAnimGraphNode_StateMachineBase* StateMachineInstance;
+		FCompilerResultsLog& MessageLog;
+	public:
+		FMachineCreator(FCompilerResultsLog& InMessageLog, UAnimGraphNode_StateMachineBase* InStateMachineInstance, int32 InMachineIndex, UAnimBlueprintGeneratedClass* InNewClass)
+			: MachineIndex(InMachineIndex)
+			, AnimBlueprintClass(InNewClass)
+			, StateMachineInstance(InStateMachineInstance)
+			, MessageLog(InMessageLog)
+		{
+			FStateMachineDebugData& MachineInfo = GetMachineSpecificDebugData();
+			MachineInfo.MachineIndex = MachineIndex;
+			MachineInfo.MachineInstanceNode = MessageLog.FindSourceObjectTypeChecked<UAnimGraphNode_StateMachineBase>(InStateMachineInstance);
+
+			StateMachineInstance->GetNode().StateMachineIndexInClass = MachineIndex;
+
+			FBakedAnimationStateMachine& BakedMachine = GetMachine();
+			BakedMachine.MachineName = StateMachineInstance->EditorStateMachineGraph->GetFName();
+			BakedMachine.InitialState = INDEX_NONE;
+		}
+
+		FBakedAnimationStateMachine& GetMachine()
+		{
+			return AnimBlueprintClass->BakedStateMachines[MachineIndex];
+		}
+
+		FStateMachineDebugData& GetMachineSpecificDebugData()
+		{
+			UAnimationStateMachineGraph* SourceGraph = MessageLog.FindSourceObjectTypeChecked<UAnimationStateMachineGraph>(StateMachineInstance->EditorStateMachineGraph);
+			return AnimBlueprintClass->GetAnimBlueprintDebugData().StateMachineDebugData.FindOrAdd(SourceGraph);
+		}
+
+		int32 FindOrAddState(UAnimStateNodeBase* StateNode)
+		{
+			if (int32* pResult = StateIndexTable.Find(StateNode))
+			{
+				return *pResult;
+			}
+			else
+			{
+				FBakedAnimationStateMachine& BakedMachine = GetMachine();
+
+				const int32 StateIndex = BakedMachine.States.Num();
+				StateIndexTable.Add(StateNode, StateIndex);
+				new (BakedMachine.States) FBakedAnimationState();
+
+				UAnimStateNodeBase* SourceNode = MessageLog.FindSourceObjectTypeChecked<UAnimStateNodeBase>(StateNode);
+				GetMachineSpecificDebugData().NodeToStateIndex.Add(SourceNode, StateIndex);
+				if (UAnimStateNode* SourceStateNode = Cast<UAnimStateNode>(SourceNode))
+				{
+					AnimBlueprintClass->GetAnimBlueprintDebugData().StateGraphToNodeMap.Add(SourceStateNode->BoundGraph, SourceStateNode);
+				}
+
+				return StateIndex;
+			}
+		}
+
+		int32 FindOrAddTransition(UAnimStateTransitionNode* TransitionNode)
+		{
+			if (int32* pResult = TransitionIndexTable.Find(TransitionNode))
+			{
+				return *pResult;
+			}
+			else
+			{
+				FBakedAnimationStateMachine& BakedMachine = GetMachine();
+
+				const int32 TransitionIndex = BakedMachine.Transitions.Num();
+				TransitionIndexTable.Add(TransitionNode, TransitionIndex);
+				new (BakedMachine.Transitions) FAnimationTransitionBetweenStates();
+
+				UAnimStateTransitionNode* SourceTransitionNode = MessageLog.FindSourceObjectTypeChecked<UAnimStateTransitionNode>(TransitionNode);
+				GetMachineSpecificDebugData().NodeToTransitionIndex.Add(SourceTransitionNode, TransitionIndex);
+				AnimBlueprintClass->GetAnimBlueprintDebugData().TransitionGraphToNodeMap.Add(SourceTransitionNode->BoundGraph, SourceTransitionNode);
+
+				if (SourceTransitionNode->CustomTransitionGraph != NULL)
+				{
+					AnimBlueprintClass->GetAnimBlueprintDebugData().TransitionBlendGraphToNodeMap.Add(SourceTransitionNode->CustomTransitionGraph, SourceTransitionNode);
+				}
+
+				return TransitionIndex;
+			}
+		}
+
+		void Validate()
+		{
+			FBakedAnimationStateMachine& BakedMachine = GetMachine();
+
+			// Make sure there is a valid entry point
+			if (BakedMachine.InitialState == INDEX_NONE)
+			{
+				MessageLog.Warning(*LOCTEXT("NoEntryNode", "There was no entry state connection in @@").ToString(), StateMachineInstance);
+				BakedMachine.InitialState = 0;
+			}
+			else
+			{
+				// Make sure the entry node is a state and not a conduit
+				if (BakedMachine.States[BakedMachine.InitialState].bIsAConduit)
+				{
+					UEdGraphNode* StateNode = GetMachineSpecificDebugData().FindNodeFromStateIndex(BakedMachine.InitialState);
+					MessageLog.Error(*LOCTEXT("BadStateEntryNode", "A conduit (@@) cannot be used as the entry node for a state machine").ToString(), StateNode);
+				}
+			}
+		}
+	};
+	
+	if (StateMachineInstance->EditorStateMachineGraph == NULL)
+	{
+		MessageLog.Error(*LOCTEXT("BadStateMachineNoGraph", "@@ does not have a corresponding graph").ToString(), StateMachineInstance);
+		return;
+	}
+
+	TMap<UAnimGraphNode_TransitionResult*, int32> AlreadyMergedTransitionList;
+
+	const int32 MachineIndex = NewAnimBlueprintClass->BakedStateMachines.Num();
+	new (NewAnimBlueprintClass->BakedStateMachines) FBakedAnimationStateMachine();
+	FMachineCreator Oven(MessageLog, StateMachineInstance, MachineIndex, NewAnimBlueprintClass);
+
+	// Map of states that contain a single player node (from state root node index to associated sequence player)
+	TMap<int32, UObject*> SimplePlayerStatesMap;
+
+	// Process all the states/transitions
+	for (auto StateNodeIt = StateMachineInstance->EditorStateMachineGraph->Nodes.CreateIterator(); StateNodeIt; ++StateNodeIt)
+	{
+		UEdGraphNode* Node = *StateNodeIt;
+
+		if (UAnimStateEntryNode* EntryNode = Cast<UAnimStateEntryNode>(Node))
+		{
+			// Handle the state graph entry
+			FBakedAnimationStateMachine& BakedMachine = Oven.GetMachine();
+			if (BakedMachine.InitialState != INDEX_NONE)
+			{
+				MessageLog.Error(*LOCTEXT("TooManyStateMachineEntryNodes", "Found an extra entry node @@").ToString(), EntryNode);
+			}
+			else if (UAnimStateNodeBase* StartState = Cast<UAnimStateNodeBase>(EntryNode->GetOutputNode()))
+			{
+				BakedMachine.InitialState = Oven.FindOrAddState(StartState);
+			}
+			else
+			{
+				MessageLog.Warning(*LOCTEXT("NoConnection", "Entry node @@ is not connected to state").ToString(), EntryNode);
+			}
+		}
+		else if (UAnimStateTransitionNode* TransitionNode = Cast<UAnimStateTransitionNode>(Node))
+		{
+			TransitionNode->ValidateNodeDuringCompilation(MessageLog);
+
+			const int32 TransitionIndex = Oven.FindOrAddTransition(TransitionNode);
+			FAnimationTransitionBetweenStates& BakedTransition = Oven.GetMachine().Transitions[TransitionIndex];
+
+			BakedTransition.CrossfadeDuration = TransitionNode->CrossfadeDuration;
+			BakedTransition.StartNotify = FindOrAddNotify(TransitionNode->TransitionStart);
+			BakedTransition.EndNotify = FindOrAddNotify(TransitionNode->TransitionEnd);
+			BakedTransition.InterruptNotify = FindOrAddNotify(TransitionNode->TransitionInterrupt);
+			BakedTransition.BlendMode = TransitionNode->BlendMode;
+			BakedTransition.CustomCurve = TransitionNode->CustomBlendCurve;
+			BakedTransition.BlendProfile = TransitionNode->BlendProfile;
+			BakedTransition.LogicType = TransitionNode->LogicType;
+
+			UAnimStateNodeBase* PreviousState = TransitionNode->GetPreviousState();
+			UAnimStateNodeBase* NextState = TransitionNode->GetNextState();
+
+			if ((PreviousState != NULL) && (NextState != NULL))
+			{
+				const int32 PreviousStateIndex = Oven.FindOrAddState(PreviousState);
+				const int32 NextStateIndex = Oven.FindOrAddState(NextState);
+
+				if (TransitionNode->Bidirectional)
+				{
+					MessageLog.Warning(*LOCTEXT("BidirectionalTransWarning", "Bidirectional transitions aren't supported yet @@").ToString(), TransitionNode);
+				}
+
+				BakedTransition.PreviousState = PreviousStateIndex;
+				BakedTransition.NextState = NextStateIndex;
+			}
+			else
+			{
+				MessageLog.Warning(*LOCTEXT("BogusTransition", "@@ is incomplete, without a previous and next state").ToString(), TransitionNode);
+				BakedTransition.PreviousState = INDEX_NONE;
+				BakedTransition.NextState = INDEX_NONE;
+			}
+		}
+		else if (UAnimStateNode* StateNode = Cast<UAnimStateNode>(Node))
+		{
+			StateNode->ValidateNodeDuringCompilation(MessageLog);
+
+			const int32 StateIndex = Oven.FindOrAddState(StateNode);
+			FBakedAnimationState& BakedState = Oven.GetMachine().States[StateIndex];
+
+			if (StateNode->BoundGraph != NULL)
+			{
+				BakedState.StateName = StateNode->BoundGraph->GetFName();
+				BakedState.StartNotify = FindOrAddNotify(StateNode->StateEntered);
+				BakedState.EndNotify = FindOrAddNotify(StateNode->StateLeft);
+				BakedState.FullyBlendedNotify = FindOrAddNotify(StateNode->StateFullyBlended);
+				BakedState.bIsAConduit = false;
+				BakedState.bAlwaysResetOnEntry = StateNode->bAlwaysResetOnEntry;
+
+				// Process the inner graph of this state
+				if (UAnimGraphNode_StateResult* AnimGraphResultNode = CastChecked<UAnimationStateGraph>(StateNode->BoundGraph)->GetResultNode())
+				{
+					BakedState.StateRootNodeIndex = ExpandGraphAndProcessNodes(StateNode->BoundGraph, AnimGraphResultNode);
+
+					// See if the state consists of a single sequence player node, and remember the index if so
+					for (UEdGraphPin* TestPin : AnimGraphResultNode->Pins)
+					{
+						if ((TestPin->Direction == EGPD_Input) && (TestPin->LinkedTo.Num() == 1))
+						{
+							if (UAnimGraphNode_SequencePlayer* SequencePlayer = Cast<UAnimGraphNode_SequencePlayer>(TestPin->LinkedTo[0]->GetOwningNode()))
+							{
+								SimplePlayerStatesMap.Add(BakedState.StateRootNodeIndex, MessageLog.FindSourceObject(SequencePlayer));
+							}
+						}
+					}
+				}
+				else
+				{
+					BakedState.StateRootNodeIndex = INDEX_NONE;
+					MessageLog.Error(*LOCTEXT("StateWithNoResult", "@@ has no result node").ToString(), StateNode);
+				}
+			}
+			else
+			{
+				BakedState.StateName = NAME_None;
+				MessageLog.Error(*LOCTEXT("StateWithBadGraph", "@@ has no bound graph").ToString(), StateNode);
+			}
+
+			// If this check fires, then something in the machine has changed causing the states array to not
+			// be a separate allocation, and a state machine inside of this one caused stuff to shift around
+			checkSlow(&BakedState == &(Oven.GetMachine().States[StateIndex]));
+		}
+		else if (UAnimStateConduitNode* ConduitNode = Cast<UAnimStateConduitNode>(Node))
+		{
+			ConduitNode->ValidateNodeDuringCompilation(MessageLog);
+
+			const int32 StateIndex = Oven.FindOrAddState(ConduitNode);
+			FBakedAnimationState& BakedState = Oven.GetMachine().States[StateIndex];
+
+			BakedState.StateName = ConduitNode->BoundGraph ? ConduitNode->BoundGraph->GetFName() : TEXT("OLD CONDUIT");
+			BakedState.bIsAConduit = true;
+			
+			if (ConduitNode->BoundGraph != NULL)
+			{
+				if (UAnimGraphNode_TransitionResult* EntryRuleResultNode = CastChecked<UAnimationTransitionGraph>(ConduitNode->BoundGraph)->GetResultNode())
+				{
+					BakedState.EntryRuleNodeIndex = ExpandGraphAndProcessNodes(ConduitNode->BoundGraph, EntryRuleResultNode);
+				}
+			}
+
+			// If this check fires, then something in the machine has changed causing the states array to not
+			// be a separate allocation, and a state machine inside of this one caused stuff to shift around
+			checkSlow(&BakedState == &(Oven.GetMachine().States[StateIndex]));
+		}
+	}
+
+	// Process transitions after all the states because getters within custom graphs may want to
+	// reference back to other states, which are only valid if they have already been baked
+	for (auto StateNodeIt = Oven.StateIndexTable.CreateIterator(); StateNodeIt; ++StateNodeIt)
+	{
+		UAnimStateNodeBase* StateNode = StateNodeIt.Key();
+		const int32 StateIndex = StateNodeIt.Value();
+
+		FBakedAnimationState& BakedState = Oven.GetMachine().States[StateIndex];
+
+		// Add indices to all player nodes
+		TArray<UEdGraph*> GraphsToCheck;
+		TArray<UAnimGraphNode_AssetPlayerBase*> AssetPlayerNodes;
+		GraphsToCheck.Add(StateNode->GetBoundGraph());
+		StateNode->GetBoundGraph()->GetAllChildrenGraphs(GraphsToCheck);
+
+		for(UEdGraph* ChildGraph : GraphsToCheck)
+		{
+			ChildGraph->GetNodesOfClass(AssetPlayerNodes);
+		}
+
+		for(UAnimGraphNode_AssetPlayerBase* Node : AssetPlayerNodes)
+		{
+			if(int32* IndexPtr = NewAnimBlueprintClass->AnimBlueprintDebugData.NodeGuidToIndexMap.Find(Node->NodeGuid))
+			{
+				BakedState.PlayerNodeIndices.Add(*IndexPtr);
+			}
+		}
+
+		// Handle all the transitions out of this node
+		TArray<class UAnimStateTransitionNode*> TransitionList;
+		StateNode->GetTransitionList(/*out*/ TransitionList, /*bWantSortedList=*/ true);
+
+		for (auto TransitionIt = TransitionList.CreateIterator(); TransitionIt; ++TransitionIt)
+		{
+			UAnimStateTransitionNode* TransitionNode = *TransitionIt;
+			const int32 TransitionIndex = Oven.FindOrAddTransition(TransitionNode);
+
+			// Validate the blend profile for this transition - incase the skeleton of the node has
+			// changed or the blend profile no longer exists.
+			TransitionNode->ValidateBlendProfile();
+
+			FBakedStateExitTransition& Rule = *new (BakedState.Transitions) FBakedStateExitTransition();
+			Rule.bDesiredTransitionReturnValue = (TransitionNode->GetPreviousState() == StateNode);
+			Rule.TransitionIndex = TransitionIndex;
+			
+			if (UAnimGraphNode_TransitionResult* TransitionResultNode = CastChecked<UAnimationTransitionGraph>(TransitionNode->BoundGraph)->GetResultNode())
+			{
+				if (int32* pIndex = AlreadyMergedTransitionList.Find(TransitionResultNode))
+				{
+					Rule.CanTakeDelegateIndex = *pIndex;
+				}
+				else
+				{
+					Rule.CanTakeDelegateIndex = ExpandGraphAndProcessNodes(TransitionNode->BoundGraph, TransitionResultNode, TransitionNode);
+					AlreadyMergedTransitionList.Add(TransitionResultNode, Rule.CanTakeDelegateIndex);
+				}
+			}
+			else
+			{
+				Rule.CanTakeDelegateIndex = INDEX_NONE;
+				MessageLog.Error(*LOCTEXT("TransitionWithNoResult", "@@ has no result node").ToString(), TransitionNode);
+			}
+
+			// Handle automatic time remaining rules
+			Rule.bAutomaticRemainingTimeRule = TransitionNode->bAutomaticRuleBasedOnSequencePlayerInState;
+
+			// Handle custom transition graphs
+			Rule.CustomResultNodeIndex = INDEX_NONE;
+			if (UAnimationCustomTransitionGraph* CustomTransitionGraph = Cast<UAnimationCustomTransitionGraph>(TransitionNode->CustomTransitionGraph))
+			{
+				TArray<UEdGraphNode*> ClonedNodes;
+				if (CustomTransitionGraph->GetResultNode())
+				{
+					Rule.CustomResultNodeIndex = ExpandGraphAndProcessNodes(TransitionNode->CustomTransitionGraph, CustomTransitionGraph->GetResultNode(), NULL, &ClonedNodes);
+				}
+
+				// Find all the pose evaluators used in this transition, save handles to them because we need to populate some pose data before executing
+				TArray<UAnimGraphNode_TransitionPoseEvaluator*> TransitionPoseList;
+				for (auto ClonedNodesIt = ClonedNodes.CreateIterator(); ClonedNodesIt; ++ClonedNodesIt)
+				{
+					UEdGraphNode* Node = *ClonedNodesIt;
+					if (UAnimGraphNode_TransitionPoseEvaluator* TypedNode = Cast<UAnimGraphNode_TransitionPoseEvaluator>(Node))
+					{
+						TransitionPoseList.Add(TypedNode);
+					}
+				}
+
+				Rule.PoseEvaluatorLinks.Empty(TransitionPoseList.Num());
+
+				for (auto TransitionPoseListIt = TransitionPoseList.CreateIterator(); TransitionPoseListIt; ++TransitionPoseListIt)
+				{
+					UAnimGraphNode_TransitionPoseEvaluator* TransitionPoseNode = *TransitionPoseListIt;
+					Rule.PoseEvaluatorLinks.Add( GetAllocationIndexOfNode(TransitionPoseNode) );
+				}
+			}
+		}
+	}
+
+	Oven.Validate();
+}
+
+void FAnimBlueprintCompilerContext::CopyTermDefaultsToDefaultObject(UObject* DefaultObject)
+{
+	Super::CopyTermDefaultsToDefaultObject(DefaultObject);
+
+	if (bIsDerivedAnimBlueprint)
+	{
+		// If we are a derived animation graph; apply any stored overrides.
+		// Restore values from the root class to catch values where the override has been removed.
+		UAnimBlueprintGeneratedClass* RootAnimClass = NewAnimBlueprintClass;
+		while(UAnimBlueprintGeneratedClass* NextClass = Cast<UAnimBlueprintGeneratedClass>(RootAnimClass->GetSuperClass()))
+		{
+			RootAnimClass = NextClass;
+		}
+		UObject* RootDefaultObject = RootAnimClass->GetDefaultObject();
+
+		for (TFieldIterator<UProperty> It(RootAnimClass) ; It; ++It)
+		{
+			UProperty* RootProp = *It;
+
+			if (UStructProperty* RootStructProp = Cast<UStructProperty>(RootProp))
+			{
+				if (RootStructProp->Struct->IsChildOf(FAnimNode_Base::StaticStruct()))
+				{
+					UStructProperty* ChildStructProp = FindField<UStructProperty>(NewAnimBlueprintClass, *RootStructProp->GetName());
+					check(ChildStructProp);
+					uint8* SourcePtr = RootStructProp->ContainerPtrToValuePtr<uint8>(RootDefaultObject);
+					uint8* DestPtr = ChildStructProp->ContainerPtrToValuePtr<uint8>(DefaultObject);
+					check(SourcePtr && DestPtr);
+					RootStructProp->CopyCompleteValue(DestPtr, SourcePtr);
+				}
+			}
+		}
+
+		// Patch the overridden values into the CDO
+
+		TArray<FAnimParentNodeAssetOverride*> AssetOverrides;
+		AnimBlueprint->GetAssetOverrides(AssetOverrides);
+		for (FAnimParentNodeAssetOverride* Override : AssetOverrides)
+		{
+			if (Override->NewAsset)
+			{
+				FAnimNode_Base* BaseNode = NewAnimBlueprintClass->GetPropertyInstance<FAnimNode_Base>(DefaultObject, Override->ParentNodeGuid, EPropertySearchMode::Hierarchy);
+				if (BaseNode)
+				{
+					BaseNode->OverrideAsset(Override->NewAsset);
+				}
+			}
+		}
+
+		return;
+	}
+
+	int32 LinkIndexCount = 0;
+	TMap<UAnimGraphNode_Base*, int32> LinkIndexMap;
+	TMap<UAnimGraphNode_Base*, uint8*> NodeBaseAddresses;
+
+	// Initialize animation nodes from their templates
+	for (TFieldIterator<UProperty> It(DefaultObject->GetClass(), EFieldIteratorFlags::ExcludeSuper); It; ++It)
+	{
+		UProperty* TargetProperty = *It;
+
+		if (UAnimGraphNode_Base* VisualAnimNode = AllocatedNodePropertiesToNodes.FindRef(TargetProperty))
+		{
+			const UStructProperty* SourceNodeProperty = VisualAnimNode->GetFNodeProperty();
+			check(SourceNodeProperty != NULL);
+			check(CastChecked<UStructProperty>(TargetProperty)->Struct == SourceNodeProperty->Struct);
+
+			uint8* DestinationPtr = TargetProperty->ContainerPtrToValuePtr<uint8>(DefaultObject);
+			uint8* SourcePtr = SourceNodeProperty->ContainerPtrToValuePtr<uint8>(VisualAnimNode);
+			TargetProperty->CopyCompleteValue(DestinationPtr, SourcePtr);
+
+			LinkIndexMap.Add(VisualAnimNode, LinkIndexCount);
+			NodeBaseAddresses.Add(VisualAnimNode, DestinationPtr);
+			++LinkIndexCount;
+		}
+	}
+
+	// And wire up node links
+	for (auto PoseLinkIt = ValidPoseLinkList.CreateIterator(); PoseLinkIt; ++PoseLinkIt)
+	{
+		FPoseLinkMappingRecord& Record = *PoseLinkIt;
+
+		UAnimGraphNode_Base* LinkingNode = Record.GetLinkingNode();
+		UAnimGraphNode_Base* LinkedNode = Record.GetLinkedNode();
+		
+		// @TODO this is quick solution for crash - if there were previous errors and some nodes were not added, they could still end here -
+		// this check avoids that and since there are already errors, compilation won't be successful.
+		// but I'd prefer stoping compilation earlier to avoid getting here in first place
+		if (LinkIndexMap.Contains(LinkingNode) && LinkIndexMap.Contains(LinkedNode))
+		{
+			const int32 SourceNodeIndex = LinkIndexMap.FindChecked(LinkingNode);
+			const int32 LinkedNodeIndex = LinkIndexMap.FindChecked(LinkedNode);
+			uint8* DestinationPtr = NodeBaseAddresses.FindChecked(LinkingNode);
+
+			Record.PatchLinkIndex(DestinationPtr, LinkedNodeIndex, SourceNodeIndex);
+		}
+	}   
+
+	// And patch evaluation function entry names
+	for (auto EvalLinkIt = ValidEvaluationHandlerList.CreateIterator(); EvalLinkIt; ++EvalLinkIt)
+	{
+		FEvaluationHandlerRecord& Record = *EvalLinkIt;
+
+		// validate fast path copy records before patching
+		Record.ValidateFastPath(DefaultObject->GetClass());
+
+		// patch either fast-path copy records or generated function names into the CDO
+		Record.PatchFunctionNameAndCopyRecordsInto(DefaultObject);
+	}
+
+	// And patch in constant values that don't need to be re-evaluated every frame
+	for (auto LiteralLinkIt = ValidAnimNodePinConstants.CreateIterator(); LiteralLinkIt; ++LiteralLinkIt)
+	{
+		FEffectiveConstantRecord& ConstantRecord = *LiteralLinkIt;
+
+		//const FString ArrayClause = (ConstantRecord.ArrayIndex != INDEX_NONE) ? FString::Printf(TEXT("[%d]"), ConstantRecord.ArrayIndex) : FString();
+		//const FString ValueClause = ConstantRecord.LiteralSourcePin->GetDefaultAsString();
+		//MessageLog.Note(*FString::Printf(TEXT("Want to set %s.%s%s to %s"), *ConstantRecord.NodeVariableProperty->GetName(), *ConstantRecord.ConstantProperty->GetName(), *ArrayClause, *ValueClause));
+
+		if (!ConstantRecord.Apply(DefaultObject))
+		{
+			MessageLog.Error(TEXT("ICE: Failed to push literal value from @@ into CDO"), ConstantRecord.LiteralSourcePin);
+		}
+	}
+}
+
+// Merges in any all ubergraph pages into the gathering ubergraph
+void FAnimBlueprintCompilerContext::MergeUbergraphPagesIn(UEdGraph* Ubergraph)
+{
+	Super::MergeUbergraphPagesIn(Ubergraph);
+
+	if (bIsDerivedAnimBlueprint)
+	{
+		// Skip any work related to an anim graph, it's all done by the parent class
+	}
+	else
+	{
+		// Move all animation graph nodes and associated pure logic chains into the consolidated event graph
+		for (int32 i = 0; i < Blueprint->FunctionGraphs.Num(); ++i)
+		{
+			UEdGraph* SourceGraph = Blueprint->FunctionGraphs[i];
+
+			if (SourceGraph->Schema->IsChildOf(UAnimationGraphSchema::StaticClass()))
+			{
+				// Merge all the animation nodes, contents, etc... into the ubergraph
+				UEdGraph* ClonedGraph = FEdGraphUtilities::CloneGraph(SourceGraph, NULL, &MessageLog, true);
+				const bool bIsLoading = Blueprint->bIsRegeneratingOnLoad || IsAsyncLoading();
+				const bool bIsCompiling = Blueprint->bBeingCompiled;
+				ClonedGraph->MoveNodesToAnotherGraph(ConsolidatedEventGraph, bIsLoading, bIsCompiling);
+			}
+		}
+
+		// Make sure we expand any split pins here before we process animation nodes.
+		for (TArray<UEdGraphNode*>::TIterator NodeIt(ConsolidatedEventGraph->Nodes); NodeIt; ++NodeIt)
+		{
+			UK2Node* K2Node = Cast<UK2Node>(*NodeIt);
+			if (K2Node != nullptr)
+			{
+				// We iterate the array in reverse so we can recombine split-pins (which modifies the pins array)
+				for (int32 PinIndex = K2Node->Pins.Num() - 1; PinIndex >= 0; --PinIndex)
+				{
+					UEdGraphPin* Pin = K2Node->Pins[PinIndex];
+					if (Pin->SubPins.Num() == 0)
+					{
+						continue;
+					}
+
+					K2Node->ExpandSplitPin(this, ConsolidatedEventGraph, Pin);
+				}
+			}
+		}
+
+		// Compile the animation graph
+		ProcessAllAnimationNodes();
+	}
+}
+
+void FAnimBlueprintCompilerContext::ProcessOneFunctionGraph(UEdGraph* SourceGraph, bool bInternalFunction)
+{
+	if (SourceGraph->Schema->IsChildOf(UAnimationGraphSchema::StaticClass()))
+	{
+		// Animation graph
+		// Do nothing, as this graph has already been processed
+	}
+	else if (SourceGraph->Schema->IsChildOf(UAnimationStateMachineSchema::StaticClass()))
+	{
+		// Animation state machine
+		// Do nothing, as this graph has already been processed
+
+		//@TODO: These should all have been moved to be child graphs by now
+		//ensure(false);
+	}
+	else
+	{
+		// Let the regular K2 compiler handle this one
+		Super::ProcessOneFunctionGraph(SourceGraph, bInternalFunction);
+	}
+}
+
+void FAnimBlueprintCompilerContext::EnsureProperGeneratedClass(UClass*& TargetUClass)
+{
+	if( TargetUClass && !((UObject*)TargetUClass)->IsA(UAnimBlueprintGeneratedClass::StaticClass()) )
+	{
+		FKismetCompilerUtilities::ConsignToOblivion(TargetUClass, Blueprint->bIsRegeneratingOnLoad);
+		TargetUClass = NULL;
+	}
+}
+
+void FAnimBlueprintCompilerContext::SpawnNewClass(const FString& NewClassName)
+{
+	NewAnimBlueprintClass = FindObject<UAnimBlueprintGeneratedClass>(Blueprint->GetOutermost(), *NewClassName);
+
+	if (NewAnimBlueprintClass == NULL)
+	{
+		NewAnimBlueprintClass = NewObject<UAnimBlueprintGeneratedClass>(Blueprint->GetOutermost(), FName(*NewClassName), RF_Public | RF_Transactional);
+	}
+	else
+	{
+		// Already existed, but wasn't linked in the Blueprint yet due to load ordering issues
+		FBlueprintCompileReinstancer::Create(NewAnimBlueprintClass);
+	}
+	NewClass = NewAnimBlueprintClass;
+}
+
+void FAnimBlueprintCompilerContext::OnNewClassSet(UBlueprintGeneratedClass* ClassToUse)
+{
+	NewAnimBlueprintClass = CastChecked<UAnimBlueprintGeneratedClass>(ClassToUse);
+}
+
+void FAnimBlueprintCompilerContext::CleanAndSanitizeClass(UBlueprintGeneratedClass* ClassToClean, UObject*& InOldCDO)
+{
+	Super::CleanAndSanitizeClass(ClassToClean, InOldCDO);
+
+	// Make sure our typed pointer is set
+	check(ClassToClean == NewClass && NewAnimBlueprintClass == NewClass);
+
+	NewAnimBlueprintClass->AnimBlueprintDebugData = FAnimBlueprintDebugData();
+
+	// Reset the baked data
+	//@TODO: Move this into PurgeClass
+	NewAnimBlueprintClass->BakedStateMachines.Empty();
+	NewAnimBlueprintClass->AnimNotifies.Empty();
+
+	NewAnimBlueprintClass->RootAnimNodeIndex = INDEX_NONE;
+	NewAnimBlueprintClass->RootAnimNodeProperty = NULL;
+	NewAnimBlueprintClass->OrderedSavedPoseIndices.Empty();
+	NewAnimBlueprintClass->AnimNodeProperties.Empty();
+
+	// Copy over runtime data from the blueprint to the class
+	NewAnimBlueprintClass->TargetSkeleton = AnimBlueprint->TargetSkeleton;
+
+	UAnimBlueprint* RootAnimBP = UAnimBlueprint::FindRootAnimBlueprint(AnimBlueprint);
+	bIsDerivedAnimBlueprint = RootAnimBP != NULL;
+
+	// Copy up data from a parent anim blueprint
+	if (bIsDerivedAnimBlueprint)
+	{
+		UAnimBlueprintGeneratedClass* RootAnimClass = CastChecked<UAnimBlueprintGeneratedClass>(RootAnimBP->GeneratedClass);
+
+		NewAnimBlueprintClass->BakedStateMachines.Append(RootAnimClass->BakedStateMachines);
+		NewAnimBlueprintClass->AnimNotifies.Append(RootAnimClass->AnimNotifies);
+		NewAnimBlueprintClass->RootAnimNodeIndex = RootAnimClass->RootAnimNodeIndex;
+		NewAnimBlueprintClass->OrderedSavedPoseIndices = RootAnimClass->OrderedSavedPoseIndices;
+	}
+}
+
+void FAnimBlueprintCompilerContext::FinishCompilingClass(UClass* Class)
+{
+	const UAnimBlueprint* PossibleRoot = UAnimBlueprint::FindRootAnimBlueprint(AnimBlueprint);
+	const UAnimBlueprint* Src = PossibleRoot ? PossibleRoot : AnimBlueprint;
+
+	UAnimBlueprintGeneratedClass* AnimBlueprintGeneratedClass = CastChecked<UAnimBlueprintGeneratedClass>(Class);
+	AnimBlueprintGeneratedClass->SyncGroupNames.Reset();
+	AnimBlueprintGeneratedClass->SyncGroupNames.Reserve(Src->Groups.Num());
+	for (const FAnimGroupInfo& GroupInfo : Src->Groups)
+	{
+		AnimBlueprintGeneratedClass->SyncGroupNames.Add(GroupInfo.Name);
+	}
+	Super::FinishCompilingClass(Class);
+}
+
+void FAnimBlueprintCompilerContext::PostCompile()
+{
+	Super::PostCompile();
+
+	UAnimBlueprintGeneratedClass* AnimBlueprintGeneratedClass = CastChecked<UAnimBlueprintGeneratedClass>(NewClass);
+
+	UAnimInstance* DefaultAnimInstance = CastChecked<UAnimInstance>(AnimBlueprintGeneratedClass->GetDefaultObject());
+
+	// copy threaded update flag to CDO
+	DefaultAnimInstance->bUseMultiThreadedAnimationUpdate = AnimBlueprint->bUseMultiThreadedAnimationUpdate;
+
+	// Verify thread-safety
+	if(GetDefault<UEngine>()->bAllowMultiThreadedAnimationUpdate && DefaultAnimInstance->bUseMultiThreadedAnimationUpdate)
+	{
+		// If we are a child anim BP, check parent classes & their CDOs
+		if (UAnimBlueprintGeneratedClass* ParentClass = Cast<UAnimBlueprintGeneratedClass>(AnimBlueprintGeneratedClass->GetSuperClass()))
+		{
+			UAnimBlueprint* ParentAnimBlueprint = Cast<UAnimBlueprint>(ParentClass->ClassGeneratedBy);
+			if (ParentAnimBlueprint && !ParentAnimBlueprint->bUseMultiThreadedAnimationUpdate)
+			{
+				DefaultAnimInstance->bUseMultiThreadedAnimationUpdate = false;
+			}
+
+			UAnimInstance* ParentDefaultObject = Cast<UAnimInstance>(ParentClass->GetDefaultObject(false));
+			if (ParentDefaultObject && !ParentDefaultObject->bUseMultiThreadedAnimationUpdate)
+			{
+				DefaultAnimInstance->bUseMultiThreadedAnimationUpdate = false;
+			}
+		}
+
+		// iterate all properties to determine validity
+		for (UStructProperty* Property : TFieldRange<UStructProperty>(AnimBlueprintGeneratedClass, EFieldIteratorFlags::IncludeSuper))
+		{
+			if(Property->Struct->IsChildOf(FAnimNode_Base::StaticStruct()))
+			{
+				FAnimNode_Base* AnimNode = Property->ContainerPtrToValuePtr<FAnimNode_Base>(DefaultAnimInstance);
+				if(!AnimNode->CanUpdateInWorkerThread())
+				{
+					MessageLog.Warning(*FText::Format(LOCTEXT("HasIncompatibleNode", "Found incompatible node \"{0}\" in blend graph. Disable threaded update or use member variable access."), FText::FromName(Property->Struct->GetFName())).ToString())
+						->AddToken(FDocumentationToken::Create(TEXT("Engine/Animation/AnimBlueprints/AnimGraph")));;
+
+					DefaultAnimInstance->bUseMultiThreadedAnimationUpdate = false;
+				}
+			}
+		}
+
+		if (FunctionList.Num() > 0)
+		{
+			// find the ubergraph in the function list
+			FKismetFunctionContext* UbergraphFunctionContext = nullptr;
+			for (FKismetFunctionContext& FunctionContext : FunctionList)
+			{
+				if (FunctionList[0].Function->GetName().StartsWith(TEXT("ExecuteUbergraph")))
+				{
+					UbergraphFunctionContext = &FunctionContext;
+					break;
+				}
+			}
+
+			if (UbergraphFunctionContext)
+			{
+				// run through the per-node compiled statements looking for struct-sets used by anim nodes
+				for (auto& StatementPair : UbergraphFunctionContext->StatementsPerNode)
+				{
+					if (UK2Node_StructMemberSet* StructMemberSetNode = Cast<UK2Node_StructMemberSet>(StatementPair.Key))
+					{
+						UObject* SourceNode = MessageLog.FindSourceObject(StructMemberSetNode);
+
+						if (SourceNode && StructMemberSetNode->StructType->IsChildOf(FAnimNode_Base::StaticStruct()))
+						{
+							for (FBlueprintCompiledStatement* Statement : StatementPair.Value)
+							{
+								if (Statement->Type == KCST_CallFunction && Statement->FunctionToCall)
+								{
+									// pure function?
+									const bool bPureFunctionCall = Statement->FunctionToCall->HasAnyFunctionFlags(FUNC_BlueprintPure);
+
+									// function called on something other than function library or anim instance?
+									UClass* FunctionClass = CastChecked<UClass>(Statement->FunctionToCall->GetOuter());
+									const bool bFunctionLibraryCall = FunctionClass->IsChildOf<UBlueprintFunctionLibrary>();
+									const bool bAnimInstanceCall = FunctionClass->IsChildOf<UAnimInstance>();
+
+									// Whitelisted/blacklisted? Some functions are not really 'pure', so we give people the opportunity to mark them up.
+									// Mark up the class if it is generally thread safe, then unsafe functions can be marked up individually. We assume
+									// that classes are unsafe by default, as well as if they are marked up NotBlueprintThreadSafe.
+									const bool bClassThreadSafe = FunctionClass->HasMetaData(TEXT("BlueprintThreadSafe"));
+									const bool bClassNotThreadSafe = FunctionClass->HasMetaData(TEXT("NotBlueprintThreadSafe")) || !FunctionClass->HasMetaData(TEXT("BlueprintThreadSafe"));
+									const bool bFunctionThreadSafe = Statement->FunctionToCall->HasMetaData(TEXT("BlueprintThreadSafe"));
+									const bool bFunctionNotThreadSafe = Statement->FunctionToCall->HasMetaData(TEXT("NotBlueprintThreadSafe"));
+
+									const bool bThreadSafe = (bClassThreadSafe && !bFunctionNotThreadSafe) || (bClassNotThreadSafe && bFunctionThreadSafe);
+
+									const bool bValidForUsage = bPureFunctionCall && bThreadSafe && (bFunctionLibraryCall || bAnimInstanceCall);
+
+									if (!bValidForUsage)
+									{
+										UEdGraphNode* FunctionNode = nullptr;
+										if (Statement->FunctionContext && Statement->FunctionContext->SourcePin)
+										{
+											FunctionNode = Statement->FunctionContext->SourcePin->GetOwningNode();
+										}
+										else if (Statement->LHS && Statement->LHS->SourcePin)
+										{
+											FunctionNode = Statement->LHS->SourcePin->GetOwningNode();
+										}
+
+										if (FunctionNode)
+										{
+											MessageLog.Warning(*LOCTEXT("NotThreadSafeWarningNodeContext", "Node @@ uses potentially thread-unsafe call @@. Disable threaded update or use a thread-safe call. Function may need BlueprintThreadSafe metadata adding.").ToString(), SourceNode, FunctionNode)
+												->AddToken(FDocumentationToken::Create(TEXT("Engine/Animation/AnimBlueprints/AnimGraph")));
+										}
+										else if (Statement->FunctionToCall)
+										{
+											MessageLog.Warning(*FText::Format(LOCTEXT("NotThreadSafeWarningFunctionContext", "Node @@ uses potentially thread-unsafe call {0}. Disable threaded update or use a thread-safe call. Function may need BlueprintThreadSafe metadata adding."), Statement->FunctionToCall->GetDisplayNameText()).ToString(), SourceNode)
+												->AddToken(FDocumentationToken::Create(TEXT("Engine/Animation/AnimBlueprints/AnimGraph")));
+										}
+										else
+										{
+											MessageLog.Warning(*LOCTEXT("NotThreadSafeWarningUnknownContext", "Node @@ uses potentially thread-unsafe call. Disable threaded update or use a thread-safe call.").ToString(), SourceNode)
+												->AddToken(FDocumentationToken::Create(TEXT("Engine/Animation/AnimBlueprints/AnimGraph")));
+										}
+
+										DefaultAnimInstance->bUseMultiThreadedAnimationUpdate = false;
+									}
+								}
+							}
+						}
+					}
+				}
+			}
+		}
+	}
+
+	for (const FEffectiveConstantRecord& ConstantRecord : ValidAnimNodePinConstants)
+	{
+		UAnimGraphNode_Base* Node = CastChecked<UAnimGraphNode_Base>(ConstantRecord.LiteralSourcePin->GetOwningNode());
+		UAnimGraphNode_Base* TrueNode = MessageLog.FindSourceObjectTypeChecked<UAnimGraphNode_Base>(Node);
+		TrueNode->BlueprintUsage = EBlueprintUsage::DoesNotUseBlueprint;
+	}
+
+	for(const FEvaluationHandlerRecord& EvaluationHandler : ValidEvaluationHandlerList)
+	{
+		if(EvaluationHandler.ServicedProperties.Num() > 0)
+		{
+			const FAnimNodeSinglePropertyHandler& Handler = EvaluationHandler.ServicedProperties.CreateConstIterator()->Value;
+			check(Handler.CopyRecords.Num() > 0);
+			check(Handler.CopyRecords[0].DestPin != nullptr);
+			UAnimGraphNode_Base* Node = CastChecked<UAnimGraphNode_Base>(Handler.CopyRecords[0].DestPin->GetOwningNode());
+			UAnimGraphNode_Base* TrueNode = MessageLog.FindSourceObjectTypeChecked<UAnimGraphNode_Base>(Node);	
+
+			FExposedValueHandler* HandlerPtr = EvaluationHandler.EvaluationHandlerProperty->ContainerPtrToValuePtr<FExposedValueHandler>(EvaluationHandler.NodeVariableProperty->ContainerPtrToValuePtr<void>(DefaultAnimInstance));
+			TrueNode->BlueprintUsage = HandlerPtr->BoundFunction != NAME_None ? EBlueprintUsage::UsesBlueprint : EBlueprintUsage::DoesNotUseBlueprint;
+
+#if WITH_EDITORONLY_DATA // ANIMINST_PostCompileValidation
+			const bool bWarnAboutBlueprintUsage = AnimBlueprint->bWarnAboutBlueprintUsage || DefaultAnimInstance->PCV_ShouldWarnAboutNodesNotUsingFastPath();
+#else
+			const bool bWarnAboutBlueprintUsage = AnimBlueprint->bWarnAboutBlueprintUsage;
+#endif
+			if (bWarnAboutBlueprintUsage && (TrueNode->BlueprintUsage == EBlueprintUsage::UsesBlueprint))
+			{
+				MessageLog.Warning(*LOCTEXT("BlueprintUsageWarning", "Node @@ uses Blueprint to update its values, access member variables directly or use a constant value for better performance.").ToString(), Node);
+			}
+		}
+	}
+
+	for (UPoseWatch* PoseWatch : AnimBlueprint->PoseWatches)
+	{
+		AnimationEditorUtils::SetPoseWatch(PoseWatch, AnimBlueprint);
+	}
+
+	// iterate all anim node and call PostCompile
+	const USkeleton* CurrentSkeleton = AnimBlueprint->TargetSkeleton;
+	for (UStructProperty* Property : TFieldRange<UStructProperty>(AnimBlueprintGeneratedClass, EFieldIteratorFlags::IncludeSuper))
+	{
+		if (Property->Struct->IsChildOf(FAnimNode_Base::StaticStruct()))
+		{
+			FAnimNode_Base* AnimNode = Property->ContainerPtrToValuePtr<FAnimNode_Base>(DefaultAnimInstance);
+			AnimNode->PostCompile(CurrentSkeleton);
+		}
+	}
+}
+
+void FAnimBlueprintCompilerContext::CreateFunctionList()
+{
+	// (These will now be processed after uber graph merge)
+
+	// Build the list of functions and do preprocessing on all of them
+	Super::CreateFunctionList();
+}
+
+void FAnimBlueprintCompilerContext::ProcessTransitionGetter(UK2Node_TransitionRuleGetter* Getter, UAnimStateTransitionNode* TransitionNode)
+{
+	// Get common elements for multiple getters
+	UEdGraphPin* OutputPin = Getter->GetOutputPin();
+
+	UEdGraphPin* SourceTimePin = NULL;
+	UAnimationAsset* AnimAsset= NULL;
+	int32 PlayerNodeIndex = INDEX_NONE;
+
+	if (UAnimGraphNode_Base* SourcePlayerNode = Getter->AssociatedAnimAssetPlayerNode)
+	{
+		// This check should never fail as the source state is always processed first before handling it's rules
+		UAnimGraphNode_Base* TrueSourceNode = MessageLog.FindSourceObjectTypeChecked<UAnimGraphNode_Base>(SourcePlayerNode);
+		UAnimGraphNode_Base* UndertypedPlayerNode = SourceNodeToProcessedNodeMap.FindRef(TrueSourceNode);
+
+		if (UndertypedPlayerNode == NULL)
+		{
+			MessageLog.Error(TEXT("ICE: Player node @@ was not processed prior to handling a transition getter @@ that used it"), SourcePlayerNode, Getter);
+			return;
+		}
+
+		// Make sure the node is still relevant
+		UEdGraph* PlayerGraph = UndertypedPlayerNode->GetGraph();
+		if (!PlayerGraph->Nodes.Contains(UndertypedPlayerNode))
+		{
+			MessageLog.Error(TEXT("@@ is not associated with a node in @@; please delete and recreate it"), Getter, PlayerGraph);
+		}
+
+		// Make sure the referenced AnimAsset player has been allocated
+		PlayerNodeIndex = GetAllocationIndexOfNode(UndertypedPlayerNode);
+		if (PlayerNodeIndex == INDEX_NONE)
+		{
+			MessageLog.Error(*LOCTEXT("BadAnimAssetNodeUsedInGetter", "@@ doesn't have a valid associated AnimAsset node.  Delete and recreate it").ToString(), Getter);
+		}
+
+		// Grab the AnimAsset, and time pin if needed
+		UScriptStruct* TimePropertyInStructType = NULL;
+		const TCHAR* TimePropertyName = NULL;
+		if (UndertypedPlayerNode->DoesSupportTimeForTransitionGetter())
+		{
+			AnimAsset = UndertypedPlayerNode->GetAnimationAsset();
+			TimePropertyInStructType = UndertypedPlayerNode->GetTimePropertyStruct();
+			TimePropertyName = UndertypedPlayerNode->GetTimePropertyName();
+		}
+		else
+		{
+			MessageLog.Error(TEXT("@@ is associated with @@, which is an unexpected type"), Getter, UndertypedPlayerNode);
+		}
+
+		bool bNeedTimePin = false;
+
+		// Determine if we need to read the current time variable from the specified sequence player
+		switch (Getter->GetterType)
+		{
+		case ETransitionGetter::AnimationAsset_GetCurrentTime:
+		case ETransitionGetter::AnimationAsset_GetCurrentTimeFraction:
+		case ETransitionGetter::AnimationAsset_GetTimeFromEnd:
+		case ETransitionGetter::AnimationAsset_GetTimeFromEndFraction:
+			bNeedTimePin = true;
+			break;
+		default:
+			bNeedTimePin = false;
+			break;
+		}
+
+		if (bNeedTimePin && (PlayerNodeIndex != INDEX_NONE) && (TimePropertyName != NULL) && (TimePropertyInStructType != NULL))
+		{
+			UProperty* NodeProperty = AllocatedPropertiesByIndex.FindChecked(PlayerNodeIndex);
+
+			// Create a struct member read node to grab the current position of the sequence player node
+			UK2Node_StructMemberGet* TimeReadNode = SpawnIntermediateNode<UK2Node_StructMemberGet>(Getter, ConsolidatedEventGraph);
+			TimeReadNode->VariableReference.SetSelfMember(NodeProperty->GetFName());
+			TimeReadNode->StructType = TimePropertyInStructType;
+
+			TimeReadNode->AllocatePinsForSingleMemberGet(TimePropertyName);
+			SourceTimePin = TimeReadNode->FindPinChecked(TimePropertyName);
+		}
+	}
+
+	// Expand it out
+	UK2Node_CallFunction* GetterHelper = NULL;
+	switch (Getter->GetterType)
+	{
+	case ETransitionGetter::AnimationAsset_GetCurrentTime:
+		if ((AnimAsset != NULL) && (SourceTimePin != NULL))
+		{
+			GetterHelper = SpawnCallAnimInstanceFunction(Getter, TEXT("GetInstanceAssetPlayerTime"));
+			GetterHelper->FindPinChecked(TEXT("AssetPlayerIndex"))->DefaultValue = FString::FromInt(PlayerNodeIndex);
+		}
+		else
+		{
+			if (Getter->AssociatedAnimAssetPlayerNode)
+			{
+				MessageLog.Error(TEXT("Please replace @@ with Get Relevant Anim Time. @@ has no animation asset"), Getter, Getter->AssociatedAnimAssetPlayerNode);
+			}
+			else
+			{
+				MessageLog.Error(TEXT("@@ is not asscociated with an asset player"), Getter);
+			}
+		}
+		break;
+	case ETransitionGetter::AnimationAsset_GetLength:
+		if (AnimAsset != NULL)
+		{
+			GetterHelper = SpawnCallAnimInstanceFunction(Getter, TEXT("GetInstanceAssetPlayerLength"));
+			GetterHelper->FindPinChecked(TEXT("AssetPlayerIndex"))->DefaultValue = FString::FromInt(PlayerNodeIndex);
+		}
+		else
+		{
+			if (Getter->AssociatedAnimAssetPlayerNode)
+			{
+				MessageLog.Error(TEXT("Please replace @@ with Get Relevant Anim Length. @@ has no animation asset"), Getter, Getter->AssociatedAnimAssetPlayerNode);
+			}
+			else
+			{
+				MessageLog.Error(TEXT("@@ is not asscociated with an asset player"), Getter);
+			}
+		}
+		break;
+	case ETransitionGetter::AnimationAsset_GetCurrentTimeFraction:
+		if ((AnimAsset != NULL) && (SourceTimePin != NULL))
+		{
+			GetterHelper = SpawnCallAnimInstanceFunction(Getter, TEXT("GetInstanceAssetPlayerTimeFraction"));
+			GetterHelper->FindPinChecked(TEXT("AssetPlayerIndex"))->DefaultValue = FString::FromInt(PlayerNodeIndex);
+		}
+		else
+		{
+			if (Getter->AssociatedAnimAssetPlayerNode)
+			{
+				MessageLog.Error(TEXT("Please replace @@ with Get Relevant Anim Time Fraction. @@ has no animation asset"), Getter, Getter->AssociatedAnimAssetPlayerNode);
+			}
+			else
+			{
+				MessageLog.Error(TEXT("@@ is not asscociated with an asset player"), Getter);
+			}
+		}
+		break;
+	case ETransitionGetter::AnimationAsset_GetTimeFromEnd:
+		if ((AnimAsset != NULL) && (SourceTimePin != NULL))
+		{
+			GetterHelper = SpawnCallAnimInstanceFunction(Getter, TEXT("GetInstanceAssetPlayerTimeFromEnd"));
+			GetterHelper->FindPinChecked(TEXT("AssetPlayerIndex"))->DefaultValue = FString::FromInt(PlayerNodeIndex);
+		}
+		else
+		{
+			if (Getter->AssociatedAnimAssetPlayerNode)
+			{
+				MessageLog.Error(TEXT("Please replace @@ with Get Relevant Anim Time Remaining. @@ has no animation asset"), Getter, Getter->AssociatedAnimAssetPlayerNode);
+			}
+			else
+			{
+				MessageLog.Error(TEXT("@@ is not asscociated with an asset player"), Getter);
+			}
+		}
+		break;
+	case ETransitionGetter::AnimationAsset_GetTimeFromEndFraction:
+		if ((AnimAsset != NULL) && (SourceTimePin != NULL))
+		{
+			GetterHelper = SpawnCallAnimInstanceFunction(Getter, TEXT("GetInstanceAssetPlayerTimeFromEndFraction"));
+			GetterHelper->FindPinChecked(TEXT("AssetPlayerIndex"))->DefaultValue = FString::FromInt(PlayerNodeIndex);
+		}
+		else
+		{
+			if (Getter->AssociatedAnimAssetPlayerNode)
+			{
+				MessageLog.Error(TEXT("Please replace @@ with Get Relevant Anim Time Remaining Fraction. @@ has no animation asset"), Getter, Getter->AssociatedAnimAssetPlayerNode);
+			}
+			else
+			{
+				MessageLog.Error(TEXT("@@ is not asscociated with an asset player"), Getter);
+			}
+		}
+		break;
+
+	case ETransitionGetter::CurrentTransitionDuration:
+		{
+			check(TransitionNode);
+			if(UAnimStateNode* SourceStateNode = MessageLog.FindSourceObjectTypeChecked<UAnimStateNode>(TransitionNode->GetPreviousState()))
+			{
+				if(UObject* SourceTransitionNode = MessageLog.FindSourceObject(TransitionNode))
+				{
+					if(FStateMachineDebugData* DebugData = NewAnimBlueprintClass->GetAnimBlueprintDebugData().StateMachineDebugData.Find(SourceStateNode->GetGraph()))
+					{
+						if(int32* pStateIndex = DebugData->NodeToStateIndex.Find(SourceStateNode))
+						{
+							const int32 StateIndex = *pStateIndex;
+							
+							// This check should never fail as all animation nodes should be processed before getters are
+							UAnimGraphNode_Base* CompiledMachineInstanceNode = SourceNodeToProcessedNodeMap.FindChecked(DebugData->MachineInstanceNode.Get());
+							const int32 MachinePropertyIndex = AllocatedAnimNodeIndices.FindChecked(CompiledMachineInstanceNode);
+							int32 TransitionPropertyIndex = INDEX_NONE;
+
+							for(TMap<TWeakObjectPtr<UEdGraphNode>, int32>::TIterator TransIt(DebugData->NodeToTransitionIndex); TransIt; ++TransIt)
+							{
+								UEdGraphNode* CurrTransNode = TransIt.Key().Get();
+								
+								if(CurrTransNode == SourceTransitionNode)
+								{
+									TransitionPropertyIndex = TransIt.Value();
+									break;
+								}
+							}
+
+							if(TransitionPropertyIndex != INDEX_NONE)
+							{
+								GetterHelper = SpawnCallAnimInstanceFunction(Getter, TEXT("GetInstanceTransitionCrossfadeDuration"));
+								GetterHelper->FindPinChecked(TEXT("MachineIndex"))->DefaultValue = FString::FromInt(MachinePropertyIndex);
+								GetterHelper->FindPinChecked(TEXT("TransitionIndex"))->DefaultValue = FString::FromInt(TransitionPropertyIndex);
+							}
+						}
+					}
+				}
+			}
+		}
+		break;
+
+	case ETransitionGetter::ArbitraryState_GetBlendWeight:
+		{
+			if (Getter->AssociatedStateNode)
+			{
+				if (UAnimStateNode* SourceStateNode = MessageLog.FindSourceObjectTypeChecked<UAnimStateNode>(Getter->AssociatedStateNode))
+				{
+					if (FStateMachineDebugData* DebugData = NewAnimBlueprintClass->GetAnimBlueprintDebugData().StateMachineDebugData.Find(SourceStateNode->GetGraph()))
+					{
+						if (int32* pStateIndex = DebugData->NodeToStateIndex.Find(SourceStateNode))
+						{
+							const int32 StateIndex = *pStateIndex;
+							//const int32 MachineIndex = DebugData->MachineIndex;
+
+							// This check should never fail as all animation nodes should be processed before getters are
+							UAnimGraphNode_Base* CompiledMachineInstanceNode = SourceNodeToProcessedNodeMap.FindChecked(DebugData->MachineInstanceNode.Get());
+							const int32 MachinePropertyIndex = AllocatedAnimNodeIndices.FindChecked(CompiledMachineInstanceNode);
+
+							GetterHelper = SpawnCallAnimInstanceFunction(Getter, TEXT("GetInstanceStateWeight"));
+							GetterHelper->FindPinChecked(TEXT("MachineIndex"))->DefaultValue = FString::FromInt(MachinePropertyIndex);
+							GetterHelper->FindPinChecked(TEXT("StateIndex"))->DefaultValue = FString::FromInt(StateIndex);
+						}
+					}
+				}
+			}
+
+			if (GetterHelper == NULL)
+			{
+				MessageLog.Error(TEXT("@@ is not associated with a valid state"), Getter);
+			}
+		}
+		break;
+
+	case ETransitionGetter::CurrentState_ElapsedTime:
+		{
+			check(TransitionNode);
+			if (UAnimStateNode* SourceStateNode = MessageLog.FindSourceObjectTypeChecked<UAnimStateNode>(TransitionNode->GetPreviousState()))
+			{
+				if (FStateMachineDebugData* DebugData = NewAnimBlueprintClass->GetAnimBlueprintDebugData().StateMachineDebugData.Find(SourceStateNode->GetGraph()))
+				{
+					// This check should never fail as all animation nodes should be processed before getters are
+					UAnimGraphNode_Base* CompiledMachineInstanceNode = SourceNodeToProcessedNodeMap.FindChecked(DebugData->MachineInstanceNode.Get());
+					const int32 MachinePropertyIndex = AllocatedAnimNodeIndices.FindChecked(CompiledMachineInstanceNode);
+
+					GetterHelper = SpawnCallAnimInstanceFunction(Getter, TEXT("GetInstanceCurrentStateElapsedTime"));
+					GetterHelper->FindPinChecked(TEXT("MachineIndex"))->DefaultValue = FString::FromInt(MachinePropertyIndex);
+				}
+			}
+			if (GetterHelper == NULL)
+			{
+				MessageLog.Error(TEXT("@@ is not associated with a valid state"), Getter);
+			}
+		}
+		break;
+
+	case ETransitionGetter::CurrentState_GetBlendWeight:
+		{
+			check(TransitionNode);
+			if (UAnimStateNode* SourceStateNode = MessageLog.FindSourceObjectTypeChecked<UAnimStateNode>(TransitionNode->GetPreviousState()))
+			{
+				{
+					if (FStateMachineDebugData* DebugData = NewAnimBlueprintClass->GetAnimBlueprintDebugData().StateMachineDebugData.Find(SourceStateNode->GetGraph()))
+					{
+						if (int32* pStateIndex = DebugData->NodeToStateIndex.Find(SourceStateNode))
+						{
+							const int32 StateIndex = *pStateIndex;
+							//const int32 MachineIndex = DebugData->MachineIndex;
+
+							// This check should never fail as all animation nodes should be processed before getters are
+							UAnimGraphNode_Base* CompiledMachineInstanceNode = SourceNodeToProcessedNodeMap.FindChecked(DebugData->MachineInstanceNode.Get());
+							const int32 MachinePropertyIndex = AllocatedAnimNodeIndices.FindChecked(CompiledMachineInstanceNode);
+
+							GetterHelper = SpawnCallAnimInstanceFunction(Getter, TEXT("GetInstanceStateWeight"));
+							GetterHelper->FindPinChecked(TEXT("MachineIndex"))->DefaultValue = FString::FromInt(MachinePropertyIndex);
+							GetterHelper->FindPinChecked(TEXT("StateIndex"))->DefaultValue = FString::FromInt(StateIndex);
+						}
+					}
+				}
+			}
+			if (GetterHelper == NULL)
+			{
+				MessageLog.Error(TEXT("@@ is not associated with a valid state"), Getter);
+			}
+		}
+		break;
+
+	default:
+		MessageLog.Error(TEXT("Unrecognized getter type on @@"), Getter);
+		break;
+	}
+
+	// Finish wiring up a call function if needed
+	if (GetterHelper != NULL)
+	{
+		check(GetterHelper->IsNodePure());
+
+		UEdGraphPin* NewReturnPin = GetterHelper->FindPinChecked(TEXT("ReturnValue"));
+		MessageLog.NotifyIntermediatePinCreation(NewReturnPin, OutputPin);
+
+		NewReturnPin->CopyPersistentDataFromOldPin(*OutputPin);
+	}
+
+	// Remove the getter from the equation
+	Getter->BreakAllNodeLinks();
+}
+
+int32 FAnimBlueprintCompilerContext::FindOrAddNotify(FAnimNotifyEvent& Notify)
+{
+	if ((Notify.NotifyName == NAME_None) && (Notify.Notify == NULL) && (Notify.NotifyStateClass == NULL))
+	{
+		// Non event, don't add it
+		return INDEX_NONE;
+	}
+
+	int32 NewIndex = INDEX_NONE;
+	for (int32 NotifyIdx = 0; NotifyIdx < NewAnimBlueprintClass->AnimNotifies.Num(); NotifyIdx++)
+	{
+		if( (NewAnimBlueprintClass->AnimNotifies[NotifyIdx].NotifyName == Notify.NotifyName) 
+			&& (NewAnimBlueprintClass->AnimNotifies[NotifyIdx].Notify == Notify.Notify) 
+			&& (NewAnimBlueprintClass->AnimNotifies[NotifyIdx].NotifyStateClass == Notify.NotifyStateClass) 
+			)
+		{
+			NewIndex = NotifyIdx;
+			break;
+		}
+	}
+
+	if (NewIndex == INDEX_NONE)
+	{
+		NewIndex = NewAnimBlueprintClass->AnimNotifies.Add(Notify);
+	}
+	return NewIndex;
+}
+
+void FAnimBlueprintCompilerContext::PostCompileDiagnostics()
+{
+	FKismetCompilerContext::PostCompileDiagnostics();
+
+#if WITH_EDITORONLY_DATA // ANIMINST_PostCompileValidation
+	// See if AnimInstance implements a PostCompileValidation Class. 
+	// If so, instantiate it, and let it perform Validation of our newly compiled AnimBlueprint.
+	if (const UAnimInstance* const DefaultAnimInstance = CastChecked<UAnimInstance>(NewAnimBlueprintClass->GetDefaultObject()))
+	{
+		if (DefaultAnimInstance->PostCompileValidationClassName.IsValid())
+		{
+			UClass* PostCompileValidationClass = LoadClass<UObject>(nullptr, *DefaultAnimInstance->PostCompileValidationClassName.ToString());
+			if (PostCompileValidationClass)
+			{
+				UAnimBlueprintPostCompileValidation* PostCompileValidation = NewObject<UAnimBlueprintPostCompileValidation>(GetTransientPackage(), PostCompileValidationClass);
+				if (PostCompileValidation)
+				{
+					FAnimBPCompileValidationParams PCV_Params(DefaultAnimInstance, NewAnimBlueprintClass, MessageLog, AllocatedNodePropertiesToNodes);
+					PostCompileValidation->DoPostCompileValidation(PCV_Params);
+				}
+			}
+		}
+	}
+#endif // WITH_EDITORONLY_DATA
+
+	if (!bIsDerivedAnimBlueprint)
+	{
+		// Run thru all nodes and make sure they like the final results
+		for (auto NodeIt = AllocatedAnimNodeIndices.CreateConstIterator(); NodeIt; ++NodeIt)
+		{
+			if (UAnimGraphNode_Base* Node = NodeIt.Key())
+			{
+				Node->ValidateAnimNodePostCompile(MessageLog, NewAnimBlueprintClass, NodeIt.Value());
+			}
+		}
+
+		//
+		bool bDisplayAnimDebug = false;
+		if (!Blueprint->bIsRegeneratingOnLoad)
+		{
+			GConfig->GetBool(TEXT("Kismet"), TEXT("CompileDisplaysAnimBlueprintBackend"), /*out*/ bDisplayAnimDebug, GEngineIni);
+
+			if (bDisplayAnimDebug)
+			{
+				DumpAnimDebugData();
+			}
+		}
+	}
+}
+
+void FAnimBlueprintCompilerContext::DumpAnimDebugData()
+{
+	// List all compiled-down nodes and their sources
+	if (NewAnimBlueprintClass->RootAnimNodeProperty == NULL)
+	{
+		return;
+	}
+
+	int32 RootIndex = INDEX_NONE;
+	NewAnimBlueprintClass->AnimNodeProperties.Find(NewAnimBlueprintClass->RootAnimNodeProperty, /*out*/ RootIndex);
+	
+	uint8* CDOBase = (uint8*)NewAnimBlueprintClass->ClassDefaultObject;
+
+	MessageLog.Note(*FString::Printf(TEXT("Anim Root is #%d"), RootIndex));
+	for (int32 Index = 0; Index < NewAnimBlueprintClass->AnimNodeProperties.Num(); ++Index)
+	{
+		UStructProperty* NodeProperty = NewAnimBlueprintClass->AnimNodeProperties[Index];
+		FString PropertyName = NodeProperty->GetName();
+		FString PropertyType = NodeProperty->Struct->GetName();
+		FString RootSuffix = (Index == RootIndex) ? TEXT(" <--- ROOT") : TEXT("");
+
+		// Print out the node
+		MessageLog.Note(*FString::Printf(TEXT("[%d] @@ [prop %s]%s"), Index, *PropertyName, *RootSuffix), AllocatedNodePropertiesToNodes.FindRef(NodeProperty));
+
+		// Print out all the node links
+		for (TFieldIterator<UProperty> PropIt(NodeProperty->Struct, EFieldIteratorFlags::IncludeSuper); PropIt; ++PropIt)
+		{
+			UProperty* ChildProp = *PropIt;
+			if (UStructProperty* ChildStructProp = Cast<UStructProperty>(ChildProp))
+			{
+				if (ChildStructProp->Struct->IsChildOf(FPoseLinkBase::StaticStruct()))
+				{
+					uint8* ChildPropertyPtr =  ChildStructProp->ContainerPtrToValuePtr<uint8>(NodeProperty->ContainerPtrToValuePtr<uint8>(CDOBase));
+					FPoseLinkBase* ChildPoseLink = (FPoseLinkBase*)ChildPropertyPtr;
+
+					if (ChildPoseLink->LinkID != INDEX_NONE)
+					{
+						UStructProperty* LinkedProperty = NewAnimBlueprintClass->AnimNodeProperties[ChildPoseLink->LinkID];
+						MessageLog.Note(*FString::Printf(TEXT("   Linked via %s to [#%d] @@"), *ChildStructProp->GetName(), ChildPoseLink->LinkID), AllocatedNodePropertiesToNodes.FindRef(LinkedProperty));
+					}
+					else
+					{
+						MessageLog.Note(*FString::Printf(TEXT("   Linked via %s to <no connection>"), *ChildStructProp->GetName()));
+					}
+				}
+			}
+		}
+	}
+
+	const int32 foo = NewAnimBlueprintClass->AnimNodeProperties.Num() - 1;
+
+	MessageLog.Note(TEXT("State machine info:"));
+	for (int32 MachineIndex = 0; MachineIndex < NewAnimBlueprintClass->BakedStateMachines.Num(); ++MachineIndex)
+	{
+		FBakedAnimationStateMachine& Machine = NewAnimBlueprintClass->BakedStateMachines[MachineIndex];
+	
+		MessageLog.Note(*FString::Printf(TEXT("Machine %s starts at state #%d (%s) and has %d states, %d transitions"),
+			*(Machine.MachineName.ToString()),
+			Machine.InitialState,
+			*(Machine.States[Machine.InitialState].StateName.ToString()),
+			Machine.States.Num(),
+			Machine.Transitions.Num()));
+
+		for (int32 StateIndex = 0; StateIndex < Machine.States.Num(); ++StateIndex)
+		{
+			FBakedAnimationState& SingleState = Machine.States[StateIndex];
+			
+			MessageLog.Note(*FString::Printf(TEXT("  State #%d is named %s, with %d exit transitions; linked to graph #%d"),
+				StateIndex,
+				*(SingleState.StateName.ToString()),
+				SingleState.Transitions.Num(),
+				foo - SingleState.StateRootNodeIndex));
+
+			for (int32 RuleIndex = 0; RuleIndex < SingleState.Transitions.Num(); ++RuleIndex)
+			{
+				FBakedStateExitTransition& ExitTransition = SingleState.Transitions[RuleIndex];
+
+				int32 TargetStateIndex = Machine.Transitions[ExitTransition.TransitionIndex].NextState;
+
+				MessageLog.Note(*FString::Printf(TEXT("    Exit trans #%d to %s uses global trans %d, wanting %s, linked to delegate #%d "),
+					RuleIndex,
+					*Machine.States[TargetStateIndex].StateName.ToString(),
+					ExitTransition.TransitionIndex,
+					ExitTransition.bDesiredTransitionReturnValue ? TEXT("TRUE") : TEXT("FALSE"),
+					foo - ExitTransition.CanTakeDelegateIndex));
+			}
+		}
+	}
+}
+
+void FAnimBlueprintCompilerContext::AutoWireAnimGetter(class UK2Node_AnimGetter* Getter, UAnimStateTransitionNode* InTransitionNode)
+{
+	UEdGraphPin* ReferencedNodeTimePin = nullptr;
+	int32 ReferencedNodeIndex = INDEX_NONE;
+	int32 SubNodeIndex = INDEX_NONE;
+	
+	UAnimGraphNode_Base* ProcessedNodeCheck = NULL;
+
+	if(UAnimGraphNode_Base* SourceNode = Getter->SourceNode)
+	{
+		UAnimGraphNode_Base* ActualSourceNode = MessageLog.FindSourceObjectTypeChecked<UAnimGraphNode_Base>(SourceNode);
+		
+		if(UAnimGraphNode_Base* ProcessedSourceNode = SourceNodeToProcessedNodeMap.FindRef(ActualSourceNode))
+		{
+			ProcessedNodeCheck = ProcessedSourceNode;
+
+			ReferencedNodeIndex = GetAllocationIndexOfNode(ProcessedSourceNode);
+
+			if(ProcessedSourceNode->DoesSupportTimeForTransitionGetter())
+			{
+				UScriptStruct* TimePropertyInStructType = ProcessedSourceNode->GetTimePropertyStruct();
+				const TCHAR* TimePropertyName = ProcessedSourceNode->GetTimePropertyName();
+
+				if(ReferencedNodeIndex != INDEX_NONE && TimePropertyName && TimePropertyInStructType)
+				{
+					UProperty* NodeProperty = AllocatedPropertiesByIndex.FindChecked(ReferencedNodeIndex);
+
+					UK2Node_StructMemberGet* ReaderNode = SpawnIntermediateNode<UK2Node_StructMemberGet>(Getter, ConsolidatedEventGraph);
+					ReaderNode->VariableReference.SetSelfMember(NodeProperty->GetFName());
+					ReaderNode->StructType = TimePropertyInStructType;
+					ReaderNode->AllocatePinsForSingleMemberGet(TimePropertyName);
+
+					ReferencedNodeTimePin = ReaderNode->FindPinChecked(TimePropertyName);
+				}
+			}
+		}
+	}
+	
+	if(Getter->SourceStateNode)
+	{
+		UObject* SourceObject = MessageLog.FindSourceObject(Getter->SourceStateNode);
+		if(UAnimStateNode* SourceStateNode = Cast<UAnimStateNode>(SourceObject))
+		{
+			if(FStateMachineDebugData* DebugData = NewAnimBlueprintClass->GetAnimBlueprintDebugData().StateMachineDebugData.Find(SourceStateNode->GetGraph()))
+			{
+				if(int32* StateIndexPtr = DebugData->NodeToStateIndex.Find(SourceStateNode))
+				{
+					SubNodeIndex = *StateIndexPtr;
+				}
+			}
+		}
+		else if(UAnimStateTransitionNode* TransitionNode = Cast<UAnimStateTransitionNode>(SourceObject))
+		{
+			if(FStateMachineDebugData* DebugData = NewAnimBlueprintClass->GetAnimBlueprintDebugData().StateMachineDebugData.Find(TransitionNode->GetGraph()))
+			{
+				if(int32* TransitionIndexPtr = DebugData->NodeToTransitionIndex.Find(TransitionNode))
+				{
+					SubNodeIndex = *TransitionIndexPtr;
+				}
+			}
+		}
+	}
+
+	check(Getter->IsNodePure());
+
+	for(UEdGraphPin* Pin : Getter->Pins)
+	{
+		// Hook up autowired parameters / pins
+		if(Pin->PinName == TEXT("CurrentTime"))
+		{
+			Pin->MakeLinkTo(ReferencedNodeTimePin);
+		}
+		else if(Pin->PinName == TEXT("AssetPlayerIndex") || Pin->PinName == TEXT("MachineIndex"))
+		{
+			Pin->DefaultValue = FString::FromInt(ReferencedNodeIndex);
+		}
+		else if(Pin->PinName == TEXT("StateIndex") || Pin->PinName == TEXT("TransitionIndex"))
+		{
+			Pin->DefaultValue = FString::FromInt(SubNodeIndex);
+		}
+	}
+}
+
+void FAnimBlueprintCompilerContext::FEvaluationHandlerRecord::PatchFunctionNameAndCopyRecordsInto(UObject* TargetObject) const
+{
+	FExposedValueHandler* HandlerPtr = EvaluationHandlerProperty->ContainerPtrToValuePtr<FExposedValueHandler>(NodeVariableProperty->ContainerPtrToValuePtr<void>(TargetObject));
+	HandlerPtr->CopyRecords.Empty();
+
+	if (IsFastPath())
+	{
+		for (const TPair<FName, FAnimNodeSinglePropertyHandler>& ServicedPropPair : ServicedProperties)
+		{
+			const FName& PropertyName = ServicedPropPair.Key;
+			const FAnimNodeSinglePropertyHandler& PropertyHandler = ServicedPropPair.Value;
+
+			for (const FPropertyCopyRecord& PropertyCopyRecord : PropertyHandler.CopyRecords)
+			{
+				  // get the correct property sizes for the type we are dealing with (array etc.)
+				  int32 DestPropertySize = PropertyCopyRecord.DestProperty->GetSize();
+				  if (UArrayProperty* DestArrayProperty = Cast<UArrayProperty>(PropertyCopyRecord.DestProperty))
+				  {
+					  DestPropertySize = DestArrayProperty->Inner->GetSize();
+				  }
+
+				  FExposedValueCopyRecord CopyRecord;
+				  CopyRecord.DestProperty = PropertyCopyRecord.DestProperty;
+				  CopyRecord.DestArrayIndex = PropertyCopyRecord.DestArrayIndex == INDEX_NONE ? 0 : PropertyCopyRecord.DestArrayIndex;
+				  CopyRecord.SourcePropertyName = PropertyCopyRecord.SourcePropertyName;
+				  CopyRecord.SourceSubPropertyName = PropertyCopyRecord.SourceSubStructPropertyName;
+				  CopyRecord.SourceArrayIndex = 0;
+				  CopyRecord.Size = DestPropertySize;
+				  CopyRecord.PostCopyOperation = PropertyCopyRecord.Operation;
+				  CopyRecord.bInstanceIsTarget = PropertyHandler.bInstanceIsTarget;
+				  HandlerPtr->CopyRecords.Add(CopyRecord);
+			}
+		}
+	}
+	else
+	{
+		// not all of our pins use copy records so we will need to call our exposed value handler
+		HandlerPtr->BoundFunction = HandlerFunctionName;
+	}
+}
+
+static UEdGraphPin* FindFirstInputPin(UEdGraphNode* InNode)
+{
+	const UAnimationGraphSchema* Schema = GetDefault<UAnimationGraphSchema>();
+
+	for(UEdGraphPin* Pin : InNode->Pins)
+	{
+		if(Pin && Pin->Direction == EGPD_Input && !Schema->IsExecPin(*Pin) && !Schema->IsSelfPin(*Pin))
+		{
+			return Pin;
+		}
+	}
+
+	return nullptr;
+}
+
+static UEdGraphNode* FollowKnots(UEdGraphPin* FromPin, UEdGraphPin*& ToPin)
+{
+	if (FromPin->LinkedTo.Num() == 0)
+	{
+		return nullptr;
+	}
+
+	UEdGraphPin* LinkedPin = FromPin->LinkedTo[0];
+	ToPin = LinkedPin;
+	if(LinkedPin)
+	{
+		UEdGraphNode* LinkedNode = LinkedPin->GetOwningNode();
+		UK2Node_Knot* KnotNode = Cast<UK2Node_Knot>(LinkedNode);
+		while(KnotNode)
+		{
+			if(UEdGraphPin* InputPin = FindFirstInputPin(KnotNode))
+			{
+				if (InputPin->LinkedTo.Num() > 0 && InputPin->LinkedTo[0])
+				{
+					ToPin = InputPin->LinkedTo[0];
+					LinkedNode = InputPin->LinkedTo[0]->GetOwningNode();
+					KnotNode = Cast<UK2Node_Knot>(LinkedNode);
+				}
+				else
+				{
+					KnotNode = nullptr;
+				}
+			}
+		}
+		return LinkedNode;
+	}
+
+	return nullptr;
+}
+
+void FAnimBlueprintCompilerContext::FEvaluationHandlerRecord::RegisterPin(UEdGraphPin* DestPin, UProperty* AssociatedProperty, int32 AssociatedPropertyArrayIndex)
+{
+	FAnimNodeSinglePropertyHandler& Handler = ServicedProperties.FindOrAdd(AssociatedProperty->GetFName());
+	Handler.CopyRecords.Emplace(DestPin, AssociatedProperty, AssociatedPropertyArrayIndex);
+}
+
+void FAnimBlueprintCompilerContext::FEvaluationHandlerRecord::BuildFastPathCopyRecords()
+{
+	if (GetDefault<UEngine>()->bOptimizeAnimBlueprintMemberVariableAccess)
+	{
+		for (TPair<FName, FAnimNodeSinglePropertyHandler>& ServicedPropPair : ServicedProperties)
+		{
+			for (FPropertyCopyRecord& CopyRecord : ServicedPropPair.Value.CopyRecords)
+			{
+				typedef bool (FAnimBlueprintCompilerContext::FEvaluationHandlerRecord::*GraphCheckerFunc)(FPropertyCopyRecord&, UEdGraphPin*);
+
+				GraphCheckerFunc GraphCheckerFuncs[] =
+				{
+					&FAnimBlueprintCompilerContext::FEvaluationHandlerRecord::CheckForVariableGet,
+					&FAnimBlueprintCompilerContext::FEvaluationHandlerRecord::CheckForLogicalNot,
+					&FAnimBlueprintCompilerContext::FEvaluationHandlerRecord::CheckForStructMemberAccess,
+				};
+
+				for (GraphCheckerFunc& CheckFunc : GraphCheckerFuncs)
+				{
+					if ((this->*CheckFunc)(CopyRecord, CopyRecord.DestPin))
+					{
+						break;
+					}
+				}
+
+				CheckForMemberOnlyAccess(CopyRecord, CopyRecord.DestPin);
+			}
+		}
+	}
+}
+
+static FName RecoverSplitStructPinName(UEdGraphPin* OutputPin)
+{
+	check(OutputPin->ParentPin);
+	
+	FString PinName = OutputPin->PinName.ToString();
+	const FString ParentPinName = OutputPin->ParentPin->PinName.ToString() + TEXT("_");
+
+	PinName.ReplaceInline(*ParentPinName, TEXT(""));
+
+	return *PinName;
+}
+
+bool FAnimBlueprintCompilerContext::FEvaluationHandlerRecord::CheckForVariableGet(FPropertyCopyRecord& CopyRecord, UEdGraphPin* DestPin)
+{
+	if(DestPin)
+	{
+		UEdGraphPin* SourcePin = nullptr;
+		if(UK2Node_VariableGet* VariableGetNode = Cast<UK2Node_VariableGet>(FollowKnots(DestPin, SourcePin)))
+		{
+			if(VariableGetNode && VariableGetNode->IsNodePure() && VariableGetNode->VariableReference.IsSelfContext())
+			{
+				if(SourcePin)
+				{
+					// variable get could be a 'split' struct
+					if(SourcePin->ParentPin != nullptr)
+					{
+						CopyRecord.SourcePropertyName = VariableGetNode->VariableReference.GetMemberName();
+						CopyRecord.SourceSubStructPropertyName = RecoverSplitStructPinName(SourcePin);
+					}
+					else
+					{
+						CopyRecord.SourcePropertyName = VariableGetNode->VariableReference.GetMemberName();
+					}
+					return true;
+				}
+			}
+		}
+	}
+
+	return false;
+}
+
+bool FAnimBlueprintCompilerContext::FEvaluationHandlerRecord::CheckForLogicalNot(FPropertyCopyRecord& CopyRecord, UEdGraphPin* DestPin)
+{
+	if(DestPin)
+	{
+		UEdGraphPin* SourcePin = nullptr;
+		UK2Node_CallFunction* CallFunctionNode = Cast<UK2Node_CallFunction>(FollowKnots(DestPin, SourcePin));
+		if(CallFunctionNode && CallFunctionNode->FunctionReference.GetMemberName() == FName(TEXT("Not_PreBool")))
+		{
+			// find and follow input pin
+			if(UEdGraphPin* InputPin = FindFirstInputPin(CallFunctionNode))
+			{
+				check(InputPin->PinType.PinCategory == UEdGraphSchema_K2::PC_Boolean);
+				if(CheckForVariableGet(CopyRecord, InputPin) || CheckForStructMemberAccess(CopyRecord, InputPin))
+				{
+					check(CopyRecord.SourcePropertyName != NAME_None);	// this should have been filled in by CheckForVariableGet() or CheckForStructMemberAccess() above
+					CopyRecord.Operation = EPostCopyOperation::LogicalNegateBool;
+					return true;
+				}
+			}
+		}
+	}
+
+	return false;
+}
+
+/** The functions that we can safely native-break */
+static const FName NativeBreakFunctionNameWhitelist[] =
+{
+	FName(TEXT("BreakVector")),
+	FName(TEXT("BreakVector2D")),
+	FName(TEXT("BreakRotator")),
+};
+
+/** Check whether a native break function can be safely used in the fast-path copy system (ie. source and dest data will be the same) */
+static bool IsWhitelistedNativeBreak(const FName& InFunctionName)
+{
+	for(const FName& FunctionName : NativeBreakFunctionNameWhitelist)
+	{
+		if(InFunctionName == FunctionName)
+		{
+			return true;
+		}
+	}
+
+	return false;
+}
+
+bool FAnimBlueprintCompilerContext::FEvaluationHandlerRecord::CheckForStructMemberAccess(FPropertyCopyRecord& CopyRecord, UEdGraphPin* DestPin)
+{
+	if(DestPin)
+	{
+		UEdGraphPin* SourcePin = nullptr;
+		if(UK2Node_BreakStruct* BreakStructNode = Cast<UK2Node_BreakStruct>(FollowKnots(DestPin, SourcePin)))
+		{
+			if(UEdGraphPin* InputPin = FindFirstInputPin(BreakStructNode))
+			{
+				if(CheckForVariableGet(CopyRecord, InputPin))
+				{
+					check(CopyRecord.SourcePropertyName != NAME_None);	// this should have been filled in by CheckForVariableGet() above
+					CopyRecord.SourceSubStructPropertyName = SourcePin->PinName;
+					return true;
+				}
+			}
+		}
+		// could be a native break
+		else if(UK2Node_CallFunction* NativeBreakNode = Cast<UK2Node_CallFunction>(FollowKnots(DestPin, SourcePin)))
+		{
+			UFunction* Function = NativeBreakNode->FunctionReference.ResolveMember<UFunction>(UKismetMathLibrary::StaticClass());
+			if(Function && Function->HasMetaData(TEXT("NativeBreakFunc")) && IsWhitelistedNativeBreak(Function->GetFName()))
+			{
+				if(UEdGraphPin* InputPin = FindFirstInputPin(NativeBreakNode))
+				{
+					if(CheckForVariableGet(CopyRecord, InputPin))
+					{
+						check(CopyRecord.SourcePropertyName != NAME_None);	// this should have been filled in by CheckForVariableGet() above
+						CopyRecord.SourceSubStructPropertyName = SourcePin->PinName;
+						return true;
+					}
+				}
+			}
+		}
+	}
+
+	return false;
+}
+
+bool FAnimBlueprintCompilerContext::FEvaluationHandlerRecord::CheckForMemberOnlyAccess(FPropertyCopyRecord& CopyRecord, UEdGraphPin* DestPin)
+{
+	const UAnimationGraphSchema* AnimGraphDefaultSchema = GetDefault<UAnimationGraphSchema>();
+
+	if(DestPin)
+	{
+		// traverse pins to leaf nodes and check for member access/pure only
+		TArray<UEdGraphPin*> PinStack;
+		PinStack.Add(DestPin);
+		while(PinStack.Num() > 0)
+		{
+			UEdGraphPin* CurrentPin = PinStack.Pop(false);
+			for(auto& LinkedPin : CurrentPin->LinkedTo)
+			{
+				UEdGraphNode* LinkedNode = LinkedPin->GetOwningNode();
+				if(LinkedNode)
+				{
+					bool bLeafNode = true;
+					for(auto& Pin : LinkedNode->Pins)
+					{
+						if(Pin != LinkedPin && Pin->Direction == EGPD_Input && !AnimGraphDefaultSchema->IsPosePin(Pin->PinType))
+						{
+							bLeafNode = false;
+							PinStack.Add(Pin);
+						}
+					}
+
+					if(bLeafNode)
+					{
+						if(UK2Node_VariableGet* LinkedVariableGetNode = Cast<UK2Node_VariableGet>(LinkedNode))
+						{
+							if(!LinkedVariableGetNode->IsNodePure() || !LinkedVariableGetNode->VariableReference.IsSelfContext())
+							{
+								// only local variable access is allowed for leaf nodes 
+								CopyRecord.InvalidateFastPath();
+							}
+						}
+						else if(UK2Node_CallFunction* CallFunctionNode = Cast<UK2Node_CallFunction>(LinkedNode))
+						{
+							if(!CallFunctionNode->IsNodePure())
+							{
+								// only allow pure function calls
+								CopyRecord.InvalidateFastPath();
+							}
+						}
+						else if(!LinkedNode->IsA<UK2Node_TransitionRuleGetter>())
+						{
+							CopyRecord.InvalidateFastPath();
+						}
+					}
+				}
+			}
+		}
+	}
+
+	return CopyRecord.IsFastPath();
+}
+
+void FAnimBlueprintCompilerContext::FEvaluationHandlerRecord::ValidateFastPath(UClass* InCompiledClass)
+{
+	for (TPair<FName, FAnimNodeSinglePropertyHandler>& ServicedPropPair : ServicedProperties)
+	{
+		for (FPropertyCopyRecord& CopyRecord : ServicedPropPair.Value.CopyRecords)
+		{
+			CopyRecord.ValidateFastPath(InCompiledClass);
+		}
+	}
+}
+
+void FAnimBlueprintCompilerContext::FPropertyCopyRecord::ValidateFastPath(UClass* InCompiledClass)
+{
+	if (IsFastPath())
+	{
+		int32 DestPropertySize = DestProperty->GetSize();
+		if (UArrayProperty* DestArrayProperty = Cast<UArrayProperty>(DestProperty))
+		{
+			DestPropertySize = DestArrayProperty->Inner->GetSize();
+		}
+
+		UProperty* SourceProperty = InCompiledClass->FindPropertyByName(SourcePropertyName);
+		if (SourceProperty)
+		{
+			if (UArrayProperty* SourceArrayProperty = Cast<UArrayProperty>(SourceProperty))
+			{
+				// We dont support arrays as source properties
+				InvalidateFastPath();
+				return;
+			}
+
+			int32 SourcePropertySize = SourceProperty->GetSize();
+			if (SourceSubStructPropertyName != NAME_None)
+			{
+				UProperty* SourceSubStructProperty = CastChecked<UStructProperty>(SourceProperty)->Struct->FindPropertyByName(SourceSubStructPropertyName);
+				if (SourceSubStructProperty)
+				{
+					SourcePropertySize = SourceSubStructProperty->GetSize();
+				}
+				else
+				{
+					InvalidateFastPath();
+					return;
+				}
+			}
+
+			if (SourcePropertySize != DestPropertySize)
+			{
+				InvalidateFastPath();
+				return;
+			}
+		}
+		else
+		{
+			InvalidateFastPath();
+			return;
+		}
+	}
+}
+
+//////////////////////////////////////////////////////////////////////////
+
+#undef LOCTEXT_NAMESPACE
diff --git a/Engine/Source/Editor/KismetCompiler/Private/AnimBlueprintCompiler.h b/Engine/Source/Editor/KismetCompiler/Private/AnimBlueprintCompiler.h
new file mode 100644
--- /dev/null
+++ b/Engine/Source/Editor/KismetCompiler/Private/AnimBlueprintCompiler.h
@@ -0,0 +1,354 @@
+// Copyright 1998-2017 Epic Games, Inc. All Rights Reserved.
+
+#pragma once
+
+#include "CoreMinimal.h"
+#include "KismetCompiler.h"
+#include "Animation/AnimNodeBase.h"
+#include "AnimGraphNode_Base.h"
+#include "KismetCompilerModule.h"
+
+class UAnimationGraphSchema;
+class UAnimGraphNode_SaveCachedPose;
+class UAnimGraphNode_StateMachineBase;
+class UAnimGraphNode_StateResult;
+class UAnimGraphNode_SubInstance;
+class UAnimGraphNode_UseCachedPose;
+class UAnimStateTransitionNode;
+class UK2Node_CallFunction;
+
+//
+// Forward declarations.
+//
+class UAnimGraphNode_SaveCachedPose;
+class UAnimGraphNode_UseCachedPose;
+class UAnimGraphNode_SubInput;
+class UAnimGraphNode_SubInstance;
+
+class UStructProperty;
+class UBlueprintGeneratedClass;
+struct FPoseLinkMappingRecord;
+
+//////////////////////////////////////////////////////////////////////////
+// FAnimBlueprintCompilerContext
+class FAnimBlueprintCompilerContext : public FKismetCompilerContext
+{
+protected:
+	typedef FKismetCompilerContext Super;
+public:
+	FAnimBlueprintCompilerContext(UAnimBlueprint* SourceSketch, FCompilerResultsLog& InMessageLog, const FKismetCompilerOptions& InCompileOptions, TArray<UObject*>* InObjLoaded);
+	virtual ~FAnimBlueprintCompilerContext();
+
+	virtual void PostCompile() override;
+
+protected:
+	// Implementation of FKismetCompilerContext interface
+	virtual void CreateClassVariablesFromBlueprint() override;
+	virtual UEdGraphSchema_K2* CreateSchema() override;
+	virtual void MergeUbergraphPagesIn(UEdGraph* Ubergraph) override;
+	virtual void ProcessOneFunctionGraph(UEdGraph* SourceGraph, bool bInternalFunction = false) override;
+	virtual void CreateFunctionList() override;
+	virtual void SpawnNewClass(const FString& NewClassName) override;
+	virtual void OnNewClassSet(UBlueprintGeneratedClass* ClassToUse) override;
+	virtual void CopyTermDefaultsToDefaultObject(UObject* DefaultObject) override;
+	virtual void PostCompileDiagnostics() override;
+	virtual void EnsureProperGeneratedClass(UClass*& TargetClass) override;
+	virtual void CleanAndSanitizeClass(UBlueprintGeneratedClass* ClassToClean, UObject*& InOldCDO) override;
+	virtual void FinishCompilingClass(UClass* Class) override;
+	// End of FKismetCompilerContext interface
+
+protected:
+	typedef TArray<UEdGraphPin*> UEdGraphPinArray;
+
+protected:
+	/** Record of a single copy operation */
+	struct FPropertyCopyRecord
+	{
+		FPropertyCopyRecord(UEdGraphPin* InDestPin, UProperty* InDestProperty, int32 InDestArrayIndex)
+			: DestPin(InDestPin)
+			, DestProperty(InDestProperty)
+			, DestArrayIndex(InDestArrayIndex)
+			, SourcePropertyName(NAME_None)
+			, SourceSubStructPropertyName(NAME_None)
+			, Operation(EPostCopyOperation::None)
+		{}
+
+		bool IsFastPath() const
+		{
+			return DestProperty != nullptr && SourcePropertyName != NAME_None;
+		}
+
+		void InvalidateFastPath()
+		{
+			SourcePropertyName = NAME_None;
+			SourceSubStructPropertyName = NAME_None;
+		}
+
+		void ValidateFastPath(UClass* InCompiledClass);
+
+		/** The destination pin we are copying to */
+		UEdGraphPin* DestPin;
+
+		/** The destination property we are copying to (on an animation node) */
+		UProperty* DestProperty;
+
+		/** The array index we use if the destination property is an array */
+		int32 DestArrayIndex;
+
+		/** The source property we are copying from (on an anim instance) */
+		FName SourcePropertyName;
+
+		/** The source sub-struct property we are copying from (if the source property is a struct property) */
+		FName SourceSubStructPropertyName;
+
+		/** Any operation we want to perform post-copy on the destination data */
+		EPostCopyOperation Operation;
+	};
+
+	// Wireup record for a single anim node property (which might be an array)
+	struct FAnimNodeSinglePropertyHandler
+	{
+		/** Copy records */
+		TArray<FPropertyCopyRecord> CopyRecords;
+
+		// If the anim instance is the container target instead of the node.
+		bool bInstanceIsTarget;
+
+		FAnimNodeSinglePropertyHandler()
+			: bInstanceIsTarget(false)
+		{
+		}
+	};
+
+	// Record for a property that was exposed as a pin, but wasn't wired up (just a literal)
+	struct FEffectiveConstantRecord
+	{
+	public:
+		// The node variable that the handler is in
+		class UStructProperty* NodeVariableProperty;
+
+		// The property within the struct to set
+		class UProperty* ConstantProperty;
+
+		// The array index if ConstantProperty is an array property, or INDEX_NONE otherwise
+		int32 ArrayIndex;
+
+		// The pin to pull the DefaultValue/DefaultObject from
+		UEdGraphPin* LiteralSourcePin;
+
+		FEffectiveConstantRecord()
+			: NodeVariableProperty(NULL)
+			, ConstantProperty(NULL)
+			, ArrayIndex(INDEX_NONE)
+			, LiteralSourcePin(NULL)
+		{
+		}
+
+		FEffectiveConstantRecord(UStructProperty* ContainingNodeProperty, UEdGraphPin* SourcePin, UProperty* SourcePinProperty, int32 SourceArrayIndex)
+			: NodeVariableProperty(ContainingNodeProperty)
+			, ConstantProperty(SourcePinProperty)
+			, ArrayIndex(SourceArrayIndex)
+			, LiteralSourcePin(SourcePin)
+		{
+		}
+
+		bool Apply(UObject* Object);
+	};
+
+	struct FEvaluationHandlerRecord
+	{
+	public:
+
+		// The node variable that the handler is in
+		UStructProperty* NodeVariableProperty;
+
+		// The specific evaluation handler inside the specified node
+		UStructProperty* EvaluationHandlerProperty;
+
+		// Whether or not our serviced properties are actually on the instance instead of the node
+		bool bServicesInstanceProperties;
+
+		// Set of properties serviced by this handler (Map from property name to the record for that property)
+		TMap<FName, FAnimNodeSinglePropertyHandler> ServicedProperties;
+
+		// The resulting function name
+		FName HandlerFunctionName;
+
+	public:
+
+		FEvaluationHandlerRecord()
+			: NodeVariableProperty(nullptr)
+			, EvaluationHandlerProperty(nullptr)
+			, bServicesInstanceProperties(false)
+			, HandlerFunctionName(NAME_None)
+		{}
+
+		bool IsFastPath() const
+		{
+			for(TMap<FName, FAnimNodeSinglePropertyHandler>::TConstIterator It(ServicedProperties); It; ++It)
+			{
+				const FAnimNodeSinglePropertyHandler& AnimNodeSinglePropertyHandler = It.Value();
+				for (const FPropertyCopyRecord& CopyRecord : AnimNodeSinglePropertyHandler.CopyRecords)
+				{
+					if (!CopyRecord.IsFastPath())
+					{
+						return false;
+					}
+				}
+			}
+
+			return true;
+		}
+
+		bool IsValid() const
+		{
+			return NodeVariableProperty != nullptr && EvaluationHandlerProperty != nullptr;
+		}
+
+		void PatchFunctionNameAndCopyRecordsInto(UObject* TargetObject) const;
+
+		void RegisterPin(UEdGraphPin* DestPin, UProperty* AssociatedProperty, int32 AssociatedPropertyArrayIndex);
+
+		UStructProperty* GetHandlerNodeProperty() const { return NodeVariableProperty; }
+
+		void BuildFastPathCopyRecords();
+
+		void ValidateFastPath(UClass* InCompiledClass);
+
+	private:
+
+		bool CheckForVariableGet(FPropertyCopyRecord& CopyRecord, UEdGraphPin* DestPin);
+
+		bool CheckForLogicalNot(FPropertyCopyRecord& CopyRecord, UEdGraphPin* DestPin);
+
+		bool CheckForStructMemberAccess(FPropertyCopyRecord& CopyRecord, UEdGraphPin* DestPin);
+
+		bool CheckForMemberOnlyAccess(FPropertyCopyRecord& CopyRecord, UEdGraphPin* DestPin);
+	};
+
+	// State machines may get processed before their inner graphs, so their node index needs to be patched up later
+	// This structure records pending fixups.
+	struct FStateRootNodeIndexFixup
+	{
+	public:
+		int32 MachineIndex;
+		int32 StateIndex;
+		UAnimGraphNode_StateResult* ResultNode;
+
+	public:
+		FStateRootNodeIndexFixup(int32 InMachineIndex, int32 InStateIndex, UAnimGraphNode_StateResult* InResultNode)
+			: MachineIndex(InMachineIndex)
+			, StateIndex(InStateIndex)
+			, ResultNode(InResultNode)
+		{
+		}
+	};
+
+protected:
+	UAnimBlueprintGeneratedClass* NewAnimBlueprintClass;
+	UAnimBlueprint* AnimBlueprint;
+
+	UAnimationGraphSchema* AnimSchema;
+
+	// Map of allocated v3 nodes that are members of the class
+	TMap<class UAnimGraphNode_Base*, UProperty*> AllocatedAnimNodes;
+	TMap<UProperty*, class UAnimGraphNode_Base*> AllocatedNodePropertiesToNodes;
+	TMap<int32, UProperty*> AllocatedPropertiesByIndex;
+
+	// Map of true source objects (user edited ones) to the cloned ones that are actually compiled
+	TMap<class UAnimGraphNode_Base*, UAnimGraphNode_Base*> SourceNodeToProcessedNodeMap;
+
+	// Index of the nodes (must match up with the runtime discovery process of nodes, which runs thru the property chain)
+	int32 AllocateNodeIndexCounter;
+	TMap<class UAnimGraphNode_Base*, int32> AllocatedAnimNodeIndices;
+
+	// Map from pose link LinkID address
+	//@TODO: Bad structure for a list of these
+	TArray<FPoseLinkMappingRecord> ValidPoseLinkList;
+
+	// List of successfully created evaluation handlers
+	TArray<FEvaluationHandlerRecord> ValidEvaluationHandlerList;
+
+	// List of animation node literals (values exposed as pins but never wired up) that need to be pushed into the CDO
+	TArray<FEffectiveConstantRecord> ValidAnimNodePinConstants;
+
+	// Map of cache name to encountered save cached pose nodes
+	TMap<FString, UAnimGraphNode_SaveCachedPose*> SaveCachedPoseNodes;
+
+	// List of getter node's we've found so the auto-wire can be deferred till after state machine compilation
+	TArray<class UK2Node_AnimGetter*> FoundGetterNodes;
+
+	// Set of used handler function names
+	TSet<FName> HandlerFunctionNames;
+
+	// True if any parent class is also generated from an animation blueprint
+	bool bIsDerivedAnimBlueprint;
+private:
+	int32 FindOrAddNotify(FAnimNotifyEvent& Notify);
+
+	UK2Node_CallFunction* SpawnCallAnimInstanceFunction(UEdGraphNode* SourceNode, FName FunctionName);
+
+	// Creates an evaluation handler for an FExposedValue property in an animation node
+	void CreateEvaluationHandlerStruct(UAnimGraphNode_Base* VisualAnimNode, FEvaluationHandlerRecord& Record);
+	void CreateEvaluationHandlerInstance(UAnimGraphNode_Base* VisualAnimNode, FEvaluationHandlerRecord& Record);
+
+	// Prunes any nodes that aren't reachable via a pose link
+	void PruneIsolatedAnimationNodes(const TArray<UAnimGraphNode_Base*>& RootSet, TArray<UAnimGraphNode_Base*>& GraphNodes);
+
+	// Compiles one animation node
+	void ProcessAnimationNode(UAnimGraphNode_Base* VisualAnimNode);
+
+	// Compiles one state machine
+	void ProcessStateMachine(UAnimGraphNode_StateMachineBase* StateMachineInstance);
+
+	// Compiles one use cached pose instance
+	void ProcessUseCachedPose(UAnimGraphNode_UseCachedPose* UseCachedPose);
+
+	// Compiles one sub instance node
+	void ProcessSubInstance(UAnimGraphNode_SubInstance* SubInstance, bool bCheckForCycles);
+
+	// Traverses subinstance links looking for slot names and state machine names, returning their count in a name map
+	typedef TMap<FName, int32> NameToCountMap;
+	void GetDuplicatedSlotAndStateNames(UAnimGraphNode_SubInstance* InSubInstance, NameToCountMap& OutStateMachineNameToCountMap, NameToCountMap& OutSlotNameToCountMap);
+
+	// Compiles an entire animation graph
+	void ProcessAllAnimationNodes();
+
+	// Convert transition getters into a function call/etc...
+	void ProcessTransitionGetter(class UK2Node_TransitionRuleGetter* Getter, UAnimStateTransitionNode* TransitionNode);
+
+	//
+	void ProcessAnimationNodesGivenRoot(TArray<UAnimGraphNode_Base*>& AnimNodeList, const TArray<UAnimGraphNode_Base*>& RootSet);
+
+	// Builds the update order list for saved pose nodes in this blueprint
+	void BuildCachedPoseNodeUpdateOrder();
+
+	// Traverses a graph to collect save pose nodes starting at InRootNode, then processes each node
+	void CachePoseNodeOrdering_StartNewTraversal(UAnimGraphNode_Base* InRootNode, TArray<UAnimGraphNode_SaveCachedPose*> &OrderedSavePoseNodes, TArray<UAnimGraphNode_Base*> VisitedRootNodes);
+
+	// Traverses a graph to collect save pose nodes starting at InAnimGraphNode, does NOT process saved pose nodes afterwards
+	void CachePoseNodeOrdering_TraverseInternal(UAnimGraphNode_Base* InAnimGraphNode, TArray<UAnimGraphNode_SaveCachedPose*> &OrderedSavePoseNodes);
+
+	// Gets all anim graph nodes that are piped into the provided node (traverses input pins)
+	void GetLinkedAnimNodes(UAnimGraphNode_Base* InGraphNode, TArray<UAnimGraphNode_Base*>& LinkedAnimNodes);
+	void GetLinkedAnimNodes_TraversePin(UEdGraphPin* InPin, TArray<UAnimGraphNode_Base*>& LinkedAnimNodes);
+	void GetLinkedAnimNodes_ProcessAnimNode(UAnimGraphNode_Base* AnimNode, TArray<UAnimGraphNode_Base*>& LinkedAnimNodes);
+
+	// Automatically fill in parameters for the specified Getter node
+	void AutoWireAnimGetter(class UK2Node_AnimGetter* Getter, UAnimStateTransitionNode* InTransitionNode);
+
+	// This function does the following steps:
+	//   Clones the nodes in the specified source graph
+	//   Merges them into the ConsolidatedEventGraph
+	//   Processes any animation nodes
+	//   Returns the index of the processed cloned version of SourceRootNode
+	//	 If supplied, will also return an array of all cloned nodes
+	int32 ExpandGraphAndProcessNodes(UEdGraph* SourceGraph, UAnimGraphNode_Base* SourceRootNode, UAnimStateTransitionNode* TransitionNode = NULL, TArray<UEdGraphNode*>* ClonedNodes = NULL);
+
+	// Dumps compiler diagnostic information
+	void DumpAnimDebugData();
+
+	// Returns the allocation index of the specified node, processing it if it was pending
+	int32 GetAllocationIndexOfNode(UAnimGraphNode_Base* VisualAnimNode);
+};
+
diff --git a/Engine/Source/Editor/KismetCompiler/Private/KismetCompiler.cpp b/Engine/Source/Editor/KismetCompiler/Private/KismetCompiler.cpp
--- a/Engine/Source/Editor/KismetCompiler/Private/KismetCompiler.cpp
+++ b/Engine/Source/Editor/KismetCompiler/Private/KismetCompiler.cpp
@@ -6,6 +6,7 @@
 
 
 #include "KismetCompiler.h"
+#include "AnimBlueprintCompiler.h"
 #include "Engine/BlueprintGeneratedClass.h"
 #include "Misc/CoreMisc.h"
 #include "Components/ActorComponent.h"
@@ -4476,5 +4477,31 @@ FString FKismetCompilerContext::GetGuid(const UEdGraphNode* Node) const
 	return Ret.ToString();
 }
 
+TMap< UClass*, CompilerContextFactoryFunction> CustomCompilerMap;
+
+TSharedPtr<FKismetCompilerContext> FKismetCompilerContext::GetCompilerForBP(UBlueprint* BP, FCompilerResultsLog& InMessageLog, const FKismetCompilerOptions& InCompileOptions)
+{
+	// Typically whatever loads the compiler module can also register it (or the module can self register). Due to load order
+	// issues anim blueprint is part of Engine and so there is no obvious place to register FAnimBlueprintCompilerContext,
+	// so I have simply hard-coded it:
+	if(UAnimBlueprint* AnimBP = Cast<UAnimBlueprint>(BP))
+	{
+		return TSharedPtr<FKismetCompilerContext>(new FAnimBlueprintCompilerContext(AnimBP, InMessageLog, InCompileOptions, nullptr));
+	}
+	else if(CompilerContextFactoryFunction* FactoryFunction = CustomCompilerMap.Find(BP->GetClass()))
+	{
+		return (*FactoryFunction)(BP, InMessageLog, InCompileOptions);
+	}
+	else
+	{
+		return TSharedPtr<FKismetCompilerContext>(new FKismetCompilerContext(BP, InMessageLog, InCompileOptions, nullptr));
+	}
+}
+
+void FKismetCompilerContext::RegisterCompilerForBP(UClass* BPClass, TFunction<TSharedPtr<FKismetCompilerContext>(UBlueprint*, FCompilerResultsLog&, const FKismetCompilerOptions&)> FactoryFunction)
+{
+	CustomCompilerMap.Add(BPClass, FactoryFunction);
+}
+
 #undef LOCTEXT_NAMESPACE
 //////////////////////////////////////////////////////////////////////////
diff --git a/Engine/Source/Editor/KismetCompiler/Private/KismetCompilerModule.cpp b/Engine/Source/Editor/KismetCompiler/Private/KismetCompilerModule.cpp
--- a/Engine/Source/Editor/KismetCompiler/Private/KismetCompilerModule.cpp
+++ b/Engine/Source/Editor/KismetCompiler/Private/KismetCompilerModule.cpp
@@ -19,6 +19,8 @@
 #include "KismetCompilerMisc.h"
 #include "KismetCompiler.h"
 
+#include "AnimBlueprintCompiler.h"
+
 #include "Kismet2/KismetDebugUtilities.h"
 #include "Kismet2/KismetReinstanceUtilities.h"
 #include "Kismet2/BlueprintEditorUtils.h"
@@ -118,9 +120,18 @@ void FKismet2CompilerModule::CompileBlueprintInner(class UBlueprint* Blueprint,
 		// if no one handles it, then use the default blueprint compiler.
 		if ( !Compiled )
 		{
-			FKismetCompilerContext Compiler(Blueprint, Results, CompileOptions, ObjLoaded);
-			Compiler.Compile();
-			check(Compiler.NewClass);
+			if ( UAnimBlueprint* AnimBlueprint = Cast<UAnimBlueprint>(Blueprint) )
+			{
+				FAnimBlueprintCompilerContext Compiler(AnimBlueprint, Results, CompileOptions, ObjLoaded);
+				Compiler.Compile();
+				check(Compiler.NewClass);
+			}
+			else
+			{
+				FKismetCompilerContext Compiler(Blueprint, Results, CompileOptions, ObjLoaded);
+				Compiler.Compile();
+				check(Compiler.NewClass);
+			}
 		}
 
 		if (bRecompileDependencies)
@@ -186,9 +197,6 @@ void FKismet2CompilerModule::CompileBlueprint(class UBlueprint* Blueprint, const
 
 	const bool bIsBrandNewBP = (Blueprint->SkeletonGeneratedClass == NULL) && (Blueprint->GeneratedClass == NULL) && (Blueprint->ParentClass != NULL) && !CompileOptions.bIsDuplicationInstigated;
 
-	// load required modules before compiling
-	Blueprint->LoadModulesRequiredForCompilation();
-
 	for ( IBlueprintCompiler* Compiler : Compilers )
 	{
 		Compiler->PreCompile(Blueprint, CompileOptions);
diff --git a/Engine/Source/Editor/KismetCompiler/Public/KismetCompiler.h b/Engine/Source/Editor/KismetCompiler/Public/KismetCompiler.h
--- a/Engine/Source/Editor/KismetCompiler/Public/KismetCompiler.h
+++ b/Engine/Source/Editor/KismetCompiler/Public/KismetCompiler.h
@@ -35,6 +35,8 @@ enum class EInternalCompilerFlags
 };
 ENUM_CLASS_FLAGS(EInternalCompilerFlags)
 
+typedef TFunction<TSharedPtr<FKismetCompilerContext>(UBlueprint*, FCompilerResultsLog&, const FKismetCompilerOptions&)> CompilerContextFactoryFunction;
+
 class KISMETCOMPILER_API FKismetCompilerContext : public FGraphCompilerContext
 {
 protected:
@@ -258,6 +260,9 @@ public:
 
 	FString GetGuid(const UEdGraphNode* Node) const;
 
+	static TSharedPtr<FKismetCompilerContext> GetCompilerForBP(UBlueprint* BP, FCompilerResultsLog& InMessageLog, const FKismetCompilerOptions& InCompileOptions);
+	static void RegisterCompilerForBP(UClass* BPClass, CompilerContextFactoryFunction FactoryFunction );
+
 	/** Ensures that all variables have valid names for compilation/replication */
 	void ValidateVariableNames();
 	
diff --git a/Engine/Source/Editor/UMGEditor/Private/UMGEditorModule.cpp b/Engine/Source/Editor/UMGEditor/Private/UMGEditorModule.cpp
--- a/Engine/Source/Editor/UMGEditor/Private/UMGEditorModule.cpp
+++ b/Engine/Source/Editor/UMGEditor/Private/UMGEditorModule.cpp
@@ -77,7 +77,7 @@ public:
 		IAssetTools& AssetTools = FModuleManager::LoadModuleChecked<FAssetToolsModule>("AssetTools").Get();
 		RegisterAssetTypeAction(AssetTools, MakeShareable(new FAssetTypeActions_WidgetBlueprint()));
 
-		FBlueprintEditorModule::RegisterCompilerGetter(UWidgetBlueprint::StaticClass(), FOnGetBlueprintCompiler::CreateStatic(&UWidgetBlueprint::GetCompilerForWidgetBP));
+		FKismetCompilerContext::RegisterCompilerForBP(UWidgetBlueprint::StaticClass(), &UWidgetBlueprint::GetCompilerForWidgetBP );
 
 		// Register with the sequencer module that we provide auto-key handlers.
 		ISequencerModule& SequencerModule = FModuleManager::Get().LoadModuleChecked<ISequencerModule>("Sequencer");
@@ -105,8 +105,6 @@ public:
 		}
 		CreatedAssetTypeActions.Empty();
 
-		FBlueprintEditorModule::UnregisterCompilerGetter(UWidgetBlueprint::StaticClass());
-
 		// Unregister sequencer track creation delegates
 		ISequencerModule* SequencerModule = FModuleManager::GetModulePtr<ISequencerModule>( "Sequencer" );
 		if ( SequencerModule != nullptr )
diff --git a/Engine/Source/Editor/UMGEditor/Private/WidgetBlueprint.cpp b/Engine/Source/Editor/UMGEditor/Private/WidgetBlueprint.cpp
--- a/Engine/Source/Editor/UMGEditor/Private/WidgetBlueprint.cpp
+++ b/Engine/Source/Editor/UMGEditor/Private/WidgetBlueprint.cpp
@@ -795,11 +795,4 @@ void UWidgetBlueprint::ForEachSourceWidgetImpl (TFunctionRef<void(UWidget*)> Fn)
 	);
 }
 
-#if WITH_EDITOR
-void UWidgetBlueprint::LoadModulesRequiredForCompilation()
-{
-	static const FName ModuleName(TEXT("UMGEditor"));
-	FModuleManager::Get().LoadModule(ModuleName);
-}
-#endif // WITH_EDITOR
 #undef LOCTEXT_NAMESPACE 
diff --git a/Engine/Source/Editor/UMGEditor/Public/WidgetBlueprint.h b/Engine/Source/Editor/UMGEditor/Public/WidgetBlueprint.h
--- a/Engine/Source/Editor/UMGEditor/Public/WidgetBlueprint.h
+++ b/Engine/Source/Editor/UMGEditor/Public/WidgetBlueprint.h
@@ -260,8 +260,4 @@ public:
 
 private:
 	void ForEachSourceWidgetImpl(TFunctionRef<void(UWidget*)> Fn) const;
-
-#if WITH_EDITOR
-	virtual void LoadModulesRequiredForCompilation() override;
-#endif 
 };
diff --git a/Engine/Source/Runtime/Engine/Classes/Animation/AnimBlueprint.h b/Engine/Source/Runtime/Engine/Classes/Animation/AnimBlueprint.h
--- a/Engine/Source/Runtime/Engine/Classes/Animation/AnimBlueprint.h
+++ b/Engine/Source/Runtime/Engine/Classes/Animation/AnimBlueprint.h
@@ -180,8 +180,4 @@ private:
 	UPROPERTY(duplicatetransient, AssetRegistrySearchable)
 	TSoftObjectPtr<class USkeletalMesh> PreviewSkeletalMesh;
 #endif
-
-#if WITH_EDITOR
-	virtual void LoadModulesRequiredForCompilation() override;
-#endif //WITH_EDITOR
 };
diff --git a/Engine/Source/Runtime/Engine/Classes/Engine/Blueprint.h b/Engine/Source/Runtime/Engine/Classes/Engine/Blueprint.h
--- a/Engine/Source/Runtime/Engine/Classes/Engine/Blueprint.h
+++ b/Engine/Source/Runtime/Engine/Classes/Engine/Blueprint.h
@@ -874,13 +874,15 @@ public:
 	void Message_Note(const FString& MessageToLog);
 	void Message_Warn(const FString& MessageToLog);
 	void Message_Error(const FString& MessageToLog);
+#endif
 	
+#if WITH_EDITORONLY_DATA
+protected:
 	/** 
-	 * Blueprint can choose to load specific modules for compilation
+	 * Blueprint can choose to load specific modules for compilation. Children are expected to call base implementation.
 	 */
 	virtual void LoadModulesRequiredForCompilation();
 #endif
-	
 };
 
 
diff --git a/Engine/Source/Runtime/Engine/Private/AnimBlueprint.cpp b/Engine/Source/Runtime/Engine/Private/AnimBlueprint.cpp
--- a/Engine/Source/Runtime/Engine/Private/AnimBlueprint.cpp
+++ b/Engine/Source/Runtime/Engine/Private/AnimBlueprint.cpp
@@ -202,9 +202,4 @@ bool UAnimBlueprint::CanRecompileWhilePlayingInEditor() const
 	return GetDefault<UEditorExperimentalSettings>()->bEnableLiveRecompilationOfAnimationBlueprints;
 }
 
-void UAnimBlueprint::LoadModulesRequiredForCompilation()
-{
-	static const FName ModuleName(TEXT("AnimationBlueprintEditor"));
-	FModuleManager::Get().LoadModule(ModuleName);
-}
 #endif
diff --git a/Engine/Source/Runtime/Engine/Private/Blueprint.cpp b/Engine/Source/Runtime/Engine/Private/Blueprint.cpp
--- a/Engine/Source/Runtime/Engine/Private/Blueprint.cpp
+++ b/Engine/Source/Runtime/Engine/Private/Blueprint.cpp
@@ -616,6 +616,8 @@ extern COREUOBJECT_API bool GBlueprintUseCompilationManager;
 
 UClass* UBlueprint::RegenerateClass(UClass* ClassToRegenerate, UObject* PreviousCDO, TArray<UObject*>& ObjLoaded)
 {
+	LoadModulesRequiredForCompilation();
+
 	if(GBlueprintUseCompilationManager)
 	{
 		// ensure that we have UProperties for any properties declared in the blueprint:
@@ -1823,9 +1825,13 @@ UEdGraph* UBlueprint::GetLastEditedUberGraph() const
 	return nullptr;
 }
 
+#endif //WITH_EDITOR
+
+#if WITH_EDITORONLY_DATA
 void UBlueprint::LoadModulesRequiredForCompilation()
 {
 	static const FName ModuleName(TEXT("KismetCompiler"));
 	FModuleManager::Get().LoadModule(ModuleName);
 }
-#endif //WITH_EDITOR
+#endif //WITH_EDITORONLY_DATA
+
