Deletion was dropped when converting this change (3729426) to edits. Also propagate 3669000, which was incorrectly integrated (because of my error failing to delete this file)

diff --git a/Engine/Source/Editor/AnimationBlueprintEditor/Private/AnimBlueprintCompiler.cpp b/Engine/Source/Editor/AnimationBlueprintEditor/Private/AnimBlueprintCompiler.cpp
deleted file mode 100644
diff --git a/Engine/Source/Editor/AnimationBlueprintEditor/Private/AnimBlueprintCompiler.h b/Engine/Source/Editor/AnimationBlueprintEditor/Private/AnimBlueprintCompiler.h
deleted file mode 100644
diff --git a/Engine/Source/Editor/KismetCompiler/Private/AnimBlueprintCompiler.cpp b/Engine/Source/Editor/KismetCompiler/Private/AnimBlueprintCompiler.cpp
--- a/Engine/Source/Editor/KismetCompiler/Private/AnimBlueprintCompiler.cpp
+++ b/Engine/Source/Editor/KismetCompiler/Private/AnimBlueprintCompiler.cpp
@@ -1701,6 +1701,167 @@
 			MessageLog.Error(TEXT("ICE: Failed to push literal value from @@ into CDO"), ConstantRecord.LiteralSourcePin);
 		}
 	}
+
+	UAnimBlueprintGeneratedClass* AnimBlueprintGeneratedClass = CastChecked<UAnimBlueprintGeneratedClass>(NewClass);
+
+	UAnimInstance* DefaultAnimInstance = CastChecked<UAnimInstance>(AnimBlueprintGeneratedClass->GetDefaultObject());
+
+	// copy threaded update flag to CDO
+	DefaultAnimInstance->bUseMultiThreadedAnimationUpdate = AnimBlueprint->bUseMultiThreadedAnimationUpdate;
+
+	// Verify thread-safety
+	if(GetDefault<UEngine>()->bAllowMultiThreadedAnimationUpdate && DefaultAnimInstance->bUseMultiThreadedAnimationUpdate)
+	{
+		// If we are a child anim BP, check parent classes & their CDOs
+		if (UAnimBlueprintGeneratedClass* ParentClass = Cast<UAnimBlueprintGeneratedClass>(AnimBlueprintGeneratedClass->GetSuperClass()))
+		{
+			UAnimBlueprint* ParentAnimBlueprint = Cast<UAnimBlueprint>(ParentClass->ClassGeneratedBy);
+			if (ParentAnimBlueprint && !ParentAnimBlueprint->bUseMultiThreadedAnimationUpdate)
+			{
+				DefaultAnimInstance->bUseMultiThreadedAnimationUpdate = false;
+			}
+
+			UAnimInstance* ParentDefaultObject = Cast<UAnimInstance>(ParentClass->GetDefaultObject(false));
+			if (ParentDefaultObject && !ParentDefaultObject->bUseMultiThreadedAnimationUpdate)
+			{
+				DefaultAnimInstance->bUseMultiThreadedAnimationUpdate = false;
+			}
+		}
+
+		// iterate all properties to determine validity
+		for (UStructProperty* Property : TFieldRange<UStructProperty>(AnimBlueprintGeneratedClass, EFieldIteratorFlags::IncludeSuper))
+		{
+			if(Property->Struct->IsChildOf(FAnimNode_Base::StaticStruct()))
+			{
+				FAnimNode_Base* AnimNode = Property->ContainerPtrToValuePtr<FAnimNode_Base>(DefaultAnimInstance);
+				if(!AnimNode->CanUpdateInWorkerThread())
+				{
+					MessageLog.Warning(*FText::Format(LOCTEXT("HasIncompatibleNode", "Found incompatible node \"{0}\" in blend graph. Disable threaded update or use member variable access."), FText::FromName(Property->Struct->GetFName())).ToString())
+						->AddToken(FDocumentationToken::Create(TEXT("Engine/Animation/AnimBlueprints/AnimGraph")));;
+
+					DefaultAnimInstance->bUseMultiThreadedAnimationUpdate = false;
+				}
+			}
+		}
+
+		if (FunctionList.Num() > 0)
+		{
+			// find the ubergraph in the function list
+			FKismetFunctionContext* UbergraphFunctionContext = nullptr;
+			for (FKismetFunctionContext& FunctionContext : FunctionList)
+			{
+				if (FunctionList[0].Function->GetName().StartsWith(TEXT("ExecuteUbergraph")))
+				{
+					UbergraphFunctionContext = &FunctionContext;
+					break;
+				}
+			}
+
+			if (UbergraphFunctionContext)
+			{
+				// run through the per-node compiled statements looking for struct-sets used by anim nodes
+				for (auto& StatementPair : UbergraphFunctionContext->StatementsPerNode)
+				{
+					if (UK2Node_StructMemberSet* StructMemberSetNode = Cast<UK2Node_StructMemberSet>(StatementPair.Key))
+					{
+						UObject* SourceNode = MessageLog.FindSourceObject(StructMemberSetNode);
+
+						if (SourceNode && StructMemberSetNode->StructType->IsChildOf(FAnimNode_Base::StaticStruct()))
+						{
+							for (FBlueprintCompiledStatement* Statement : StatementPair.Value)
+							{
+								if (Statement->Type == KCST_CallFunction && Statement->FunctionToCall)
+								{
+									// pure function?
+									const bool bPureFunctionCall = Statement->FunctionToCall->HasAnyFunctionFlags(FUNC_BlueprintPure);
+
+									// function called on something other than function library or anim instance?
+									UClass* FunctionClass = CastChecked<UClass>(Statement->FunctionToCall->GetOuter());
+									const bool bFunctionLibraryCall = FunctionClass->IsChildOf<UBlueprintFunctionLibrary>();
+									const bool bAnimInstanceCall = FunctionClass->IsChildOf<UAnimInstance>();
+
+									// Whitelisted/blacklisted? Some functions are not really 'pure', so we give people the opportunity to mark them up.
+									// Mark up the class if it is generally thread safe, then unsafe functions can be marked up individually. We assume
+									// that classes are unsafe by default, as well as if they are marked up NotBlueprintThreadSafe.
+									const bool bClassThreadSafe = FunctionClass->HasMetaData(TEXT("BlueprintThreadSafe"));
+									const bool bClassNotThreadSafe = FunctionClass->HasMetaData(TEXT("NotBlueprintThreadSafe")) || !FunctionClass->HasMetaData(TEXT("BlueprintThreadSafe"));
+									const bool bFunctionThreadSafe = Statement->FunctionToCall->HasMetaData(TEXT("BlueprintThreadSafe"));
+									const bool bFunctionNotThreadSafe = Statement->FunctionToCall->HasMetaData(TEXT("NotBlueprintThreadSafe"));
+
+									const bool bThreadSafe = (bClassThreadSafe && !bFunctionNotThreadSafe) || (bClassNotThreadSafe && bFunctionThreadSafe);
+
+									const bool bValidForUsage = bPureFunctionCall && bThreadSafe && (bFunctionLibraryCall || bAnimInstanceCall);
+
+									if (!bValidForUsage)
+									{
+										UEdGraphNode* FunctionNode = nullptr;
+										if (Statement->FunctionContext && Statement->FunctionContext->SourcePin)
+										{
+											FunctionNode = Statement->FunctionContext->SourcePin->GetOwningNode();
+										}
+										else if (Statement->LHS && Statement->LHS->SourcePin)
+										{
+											FunctionNode = Statement->LHS->SourcePin->GetOwningNode();
+										}
+
+										if (FunctionNode)
+										{
+											MessageLog.Warning(*LOCTEXT("NotThreadSafeWarningNodeContext", "Node @@ uses potentially thread-unsafe call @@. Disable threaded update or use a thread-safe call. Function may need BlueprintThreadSafe metadata adding.").ToString(), SourceNode, FunctionNode)
+												->AddToken(FDocumentationToken::Create(TEXT("Engine/Animation/AnimBlueprints/AnimGraph")));
+										}
+										else if (Statement->FunctionToCall)
+										{
+											MessageLog.Warning(*FText::Format(LOCTEXT("NotThreadSafeWarningFunctionContext", "Node @@ uses potentially thread-unsafe call {0}. Disable threaded update or use a thread-safe call. Function may need BlueprintThreadSafe metadata adding."), Statement->FunctionToCall->GetDisplayNameText()).ToString(), SourceNode)
+												->AddToken(FDocumentationToken::Create(TEXT("Engine/Animation/AnimBlueprints/AnimGraph")));
+										}
+										else
+										{
+											MessageLog.Warning(*LOCTEXT("NotThreadSafeWarningUnknownContext", "Node @@ uses potentially thread-unsafe call. Disable threaded update or use a thread-safe call.").ToString(), SourceNode)
+												->AddToken(FDocumentationToken::Create(TEXT("Engine/Animation/AnimBlueprints/AnimGraph")));
+										}
+
+										DefaultAnimInstance->bUseMultiThreadedAnimationUpdate = false;
+									}
+								}
+							}
+						}
+					}
+				}
+			}
+		}
+	}
+
+	for (const FEffectiveConstantRecord& ConstantRecord : ValidAnimNodePinConstants)
+	{
+		UAnimGraphNode_Base* Node = CastChecked<UAnimGraphNode_Base>(ConstantRecord.LiteralSourcePin->GetOwningNode());
+		UAnimGraphNode_Base* TrueNode = MessageLog.FindSourceObjectTypeChecked<UAnimGraphNode_Base>(Node);
+		TrueNode->BlueprintUsage = EBlueprintUsage::DoesNotUseBlueprint;
+	}
+
+	for(const FEvaluationHandlerRecord& EvaluationHandler : ValidEvaluationHandlerList)
+	{
+		if(EvaluationHandler.ServicedProperties.Num() > 0)
+		{
+			const FAnimNodeSinglePropertyHandler& Handler = EvaluationHandler.ServicedProperties.CreateConstIterator()->Value;
+			check(Handler.CopyRecords.Num() > 0);
+			check(Handler.CopyRecords[0].DestPin != nullptr);
+			UAnimGraphNode_Base* Node = CastChecked<UAnimGraphNode_Base>(Handler.CopyRecords[0].DestPin->GetOwningNode());
+			UAnimGraphNode_Base* TrueNode = MessageLog.FindSourceObjectTypeChecked<UAnimGraphNode_Base>(Node);	
+
+			FExposedValueHandler* HandlerPtr = EvaluationHandler.EvaluationHandlerProperty->ContainerPtrToValuePtr<FExposedValueHandler>(EvaluationHandler.NodeVariableProperty->ContainerPtrToValuePtr<void>(DefaultAnimInstance));
+			TrueNode->BlueprintUsage = HandlerPtr->BoundFunction != NAME_None ? EBlueprintUsage::UsesBlueprint : EBlueprintUsage::DoesNotUseBlueprint; 
+
+#if WITH_EDITORONLY_DATA // ANIMINST_PostCompileValidation
+			const bool bWarnAboutBlueprintUsage = AnimBlueprint->bWarnAboutBlueprintUsage || DefaultAnimInstance->PCV_ShouldWarnAboutNodesNotUsingFastPath();
+#else
+			const bool bWarnAboutBlueprintUsage = AnimBlueprint->bWarnAboutBlueprintUsage;
+#endif
+			if (bWarnAboutBlueprintUsage && (TrueNode->BlueprintUsage == EBlueprintUsage::UsesBlueprint))
+			{
+				MessageLog.Warning(*LOCTEXT("BlueprintUsageWarning", "Node @@ uses Blueprint to update its values, access member variables directly or use a constant value for better performance.").ToString(), Node);
+			}
+		}
+	}
 }
 
 // Merges in any all ubergraph pages into the gathering ubergraph
@@ -1866,163 +2027,6 @@
 
 	UAnimInstance* DefaultAnimInstance = CastChecked<UAnimInstance>(AnimBlueprintGeneratedClass->GetDefaultObject());
 
-	// copy threaded update flag to CDO
-	DefaultAnimInstance->bUseMultiThreadedAnimationUpdate = AnimBlueprint->bUseMultiThreadedAnimationUpdate;
-
-	// Verify thread-safety
-	if(GetDefault<UEngine>()->bAllowMultiThreadedAnimationUpdate && DefaultAnimInstance->bUseMultiThreadedAnimationUpdate)
-	{
-		// If we are a child anim BP, check parent classes & their CDOs
-		if (UAnimBlueprintGeneratedClass* ParentClass = Cast<UAnimBlueprintGeneratedClass>(AnimBlueprintGeneratedClass->GetSuperClass()))
-		{
-			UAnimBlueprint* ParentAnimBlueprint = Cast<UAnimBlueprint>(ParentClass->ClassGeneratedBy);
-			if (ParentAnimBlueprint && !ParentAnimBlueprint->bUseMultiThreadedAnimationUpdate)
-			{
-				DefaultAnimInstance->bUseMultiThreadedAnimationUpdate = false;
-			}
-
-			UAnimInstance* ParentDefaultObject = Cast<UAnimInstance>(ParentClass->GetDefaultObject(false));
-			if (ParentDefaultObject && !ParentDefaultObject->bUseMultiThreadedAnimationUpdate)
-			{
-				DefaultAnimInstance->bUseMultiThreadedAnimationUpdate = false;
-			}
-		}
-
-		// iterate all properties to determine validity
-		for (UStructProperty* Property : TFieldRange<UStructProperty>(AnimBlueprintGeneratedClass, EFieldIteratorFlags::IncludeSuper))
-		{
-			if(Property->Struct->IsChildOf(FAnimNode_Base::StaticStruct()))
-			{
-				FAnimNode_Base* AnimNode = Property->ContainerPtrToValuePtr<FAnimNode_Base>(DefaultAnimInstance);
-				if(!AnimNode->CanUpdateInWorkerThread())
-				{
-					MessageLog.Warning(*FText::Format(LOCTEXT("HasIncompatibleNode", "Found incompatible node \"{0}\" in blend graph. Disable threaded update or use member variable access."), FText::FromName(Property->Struct->GetFName())).ToString())
-						->AddToken(FDocumentationToken::Create(TEXT("Engine/Animation/AnimBlueprints/AnimGraph")));;
-
-					DefaultAnimInstance->bUseMultiThreadedAnimationUpdate = false;
-				}
-			}
-		}
-
-		if (FunctionList.Num() > 0)
-		{
-			// find the ubergraph in the function list
-			FKismetFunctionContext* UbergraphFunctionContext = nullptr;
-			for (FKismetFunctionContext& FunctionContext : FunctionList)
-			{
-				if (FunctionList[0].Function->GetName().StartsWith(TEXT("ExecuteUbergraph")))
-				{
-					UbergraphFunctionContext = &FunctionContext;
-					break;
-				}
-			}
-
-			if (UbergraphFunctionContext)
-			{
-				// run through the per-node compiled statements looking for struct-sets used by anim nodes
-				for (auto& StatementPair : UbergraphFunctionContext->StatementsPerNode)
-				{
-					if (UK2Node_StructMemberSet* StructMemberSetNode = Cast<UK2Node_StructMemberSet>(StatementPair.Key))
-					{
-						UObject* SourceNode = MessageLog.FindSourceObject(StructMemberSetNode);
-
-						if (SourceNode && StructMemberSetNode->StructType->IsChildOf(FAnimNode_Base::StaticStruct()))
-						{
-							for (FBlueprintCompiledStatement* Statement : StatementPair.Value)
-							{
-								if (Statement->Type == KCST_CallFunction && Statement->FunctionToCall)
-								{
-									// pure function?
-									const bool bPureFunctionCall = Statement->FunctionToCall->HasAnyFunctionFlags(FUNC_BlueprintPure);
-
-									// function called on something other than function library or anim instance?
-									UClass* FunctionClass = CastChecked<UClass>(Statement->FunctionToCall->GetOuter());
-									const bool bFunctionLibraryCall = FunctionClass->IsChildOf<UBlueprintFunctionLibrary>();
-									const bool bAnimInstanceCall = FunctionClass->IsChildOf<UAnimInstance>();
-
-									// Whitelisted/blacklisted? Some functions are not really 'pure', so we give people the opportunity to mark them up.
-									// Mark up the class if it is generally thread safe, then unsafe functions can be marked up individually. We assume
-									// that classes are unsafe by default, as well as if they are marked up NotBlueprintThreadSafe.
-									const bool bClassThreadSafe = FunctionClass->HasMetaData(TEXT("BlueprintThreadSafe"));
-									const bool bClassNotThreadSafe = FunctionClass->HasMetaData(TEXT("NotBlueprintThreadSafe")) || !FunctionClass->HasMetaData(TEXT("BlueprintThreadSafe"));
-									const bool bFunctionThreadSafe = Statement->FunctionToCall->HasMetaData(TEXT("BlueprintThreadSafe"));
-									const bool bFunctionNotThreadSafe = Statement->FunctionToCall->HasMetaData(TEXT("NotBlueprintThreadSafe"));
-
-									const bool bThreadSafe = (bClassThreadSafe && !bFunctionNotThreadSafe) || (bClassNotThreadSafe && bFunctionThreadSafe);
-
-									const bool bValidForUsage = bPureFunctionCall && bThreadSafe && (bFunctionLibraryCall || bAnimInstanceCall);
-
-									if (!bValidForUsage)
-									{
-										UEdGraphNode* FunctionNode = nullptr;
-										if (Statement->FunctionContext && Statement->FunctionContext->SourcePin)
-										{
-											FunctionNode = Statement->FunctionContext->SourcePin->GetOwningNode();
-										}
-										else if (Statement->LHS && Statement->LHS->SourcePin)
-										{
-											FunctionNode = Statement->LHS->SourcePin->GetOwningNode();
-										}
-
-										if (FunctionNode)
-										{
-											MessageLog.Warning(*LOCTEXT("NotThreadSafeWarningNodeContext", "Node @@ uses potentially thread-unsafe call @@. Disable threaded update or use a thread-safe call. Function may need BlueprintThreadSafe metadata adding.").ToString(), SourceNode, FunctionNode)
-												->AddToken(FDocumentationToken::Create(TEXT("Engine/Animation/AnimBlueprints/AnimGraph")));
-										}
-										else if (Statement->FunctionToCall)
-										{
-											MessageLog.Warning(*FText::Format(LOCTEXT("NotThreadSafeWarningFunctionContext", "Node @@ uses potentially thread-unsafe call {0}. Disable threaded update or use a thread-safe call. Function may need BlueprintThreadSafe metadata adding."), Statement->FunctionToCall->GetDisplayNameText()).ToString(), SourceNode)
-												->AddToken(FDocumentationToken::Create(TEXT("Engine/Animation/AnimBlueprints/AnimGraph")));
-										}
-										else
-										{
-											MessageLog.Warning(*LOCTEXT("NotThreadSafeWarningUnknownContext", "Node @@ uses potentially thread-unsafe call. Disable threaded update or use a thread-safe call.").ToString(), SourceNode)
-												->AddToken(FDocumentationToken::Create(TEXT("Engine/Animation/AnimBlueprints/AnimGraph")));
-										}
-
-										DefaultAnimInstance->bUseMultiThreadedAnimationUpdate = false;
-									}
-								}
-							}
-						}
-					}
-				}
-			}
-		}
-	}
-
-	for (const FEffectiveConstantRecord& ConstantRecord : ValidAnimNodePinConstants)
-	{
-		UAnimGraphNode_Base* Node = CastChecked<UAnimGraphNode_Base>(ConstantRecord.LiteralSourcePin->GetOwningNode());
-		UAnimGraphNode_Base* TrueNode = MessageLog.FindSourceObjectTypeChecked<UAnimGraphNode_Base>(Node);
-		TrueNode->BlueprintUsage = EBlueprintUsage::DoesNotUseBlueprint;
-	}
-
-	for(const FEvaluationHandlerRecord& EvaluationHandler : ValidEvaluationHandlerList)
-	{
-		if(EvaluationHandler.ServicedProperties.Num() > 0)
-		{
-			const FAnimNodeSinglePropertyHandler& Handler = EvaluationHandler.ServicedProperties.CreateConstIterator()->Value;
-			check(Handler.CopyRecords.Num() > 0);
-			check(Handler.CopyRecords[0].DestPin != nullptr);
-			UAnimGraphNode_Base* Node = CastChecked<UAnimGraphNode_Base>(Handler.CopyRecords[0].DestPin->GetOwningNode());
-			UAnimGraphNode_Base* TrueNode = MessageLog.FindSourceObjectTypeChecked<UAnimGraphNode_Base>(Node);	
-
-			FExposedValueHandler* HandlerPtr = EvaluationHandler.EvaluationHandlerProperty->ContainerPtrToValuePtr<FExposedValueHandler>(EvaluationHandler.NodeVariableProperty->ContainerPtrToValuePtr<void>(DefaultAnimInstance));
-			TrueNode->BlueprintUsage = HandlerPtr->BoundFunction != NAME_None ? EBlueprintUsage::UsesBlueprint : EBlueprintUsage::DoesNotUseBlueprint;
-
-#if WITH_EDITORONLY_DATA // ANIMINST_PostCompileValidation
-			const bool bWarnAboutBlueprintUsage = AnimBlueprint->bWarnAboutBlueprintUsage || DefaultAnimInstance->PCV_ShouldWarnAboutNodesNotUsingFastPath();
-#else
-			const bool bWarnAboutBlueprintUsage = AnimBlueprint->bWarnAboutBlueprintUsage;
-#endif
-			if (bWarnAboutBlueprintUsage && (TrueNode->BlueprintUsage == EBlueprintUsage::UsesBlueprint))
-			{
-				MessageLog.Warning(*LOCTEXT("BlueprintUsageWarning", "Node @@ uses Blueprint to update its values, access member variables directly or use a constant value for better performance.").ToString(), Node);
-			}
-		}
-	}
-
 	for (UPoseWatch* PoseWatch : AnimBlueprint->PoseWatches)
 	{
 		AnimationEditorUtils::SetPoseWatch(PoseWatch, AnimBlueprint);
