Copying //UE4/Dev-Framework to //UE4/Dev-Main (Source: //UE4/Dev-Framework @ 3779049)

diff --git a/Engine/Source/Editor/EditorStyle/Private/SlateEditorStyle.cpp b/Engine/Source/Editor/EditorStyle/Private/SlateEditorStyle.cpp
--- a/Engine/Source/Editor/EditorStyle/Private/SlateEditorStyle.cpp
+++ b/Engine/Source/Editor/EditorStyle/Private/SlateEditorStyle.cpp
@@ -5601,9 +5601,12 @@ void FSlateEditorStyle::FStyle::SetupPersonaStyle()
 
 		Set( "PlayWorld.ShowCurrentStatement", new IMAGE_BRUSH( "Icons/icon_findnode_40x", Icon40x40 ) );
 		Set( "PlayWorld.ShowCurrentStatement.Small", new IMAGE_BRUSH( "Icons/icon_findnode_40x", Icon20x20 ) );
-		Set( "PlayWorld.StepInto", new IMAGE_BRUSH( "Icons/icon_step_40x", Icon40x40 ) );
-		Set( "PlayWorld.StepInto.Small", new IMAGE_BRUSH( "Icons/icon_step_40x", Icon20x20 ) );
-		Set( "PlayWorld.StepOver", new IMAGE_BRUSH( "Old/Kismet2/Debugger_StepOver", Icon40x40 ) );
+		Set( "PlayWorld.StepOut", new IMAGE_BRUSH("Icons/icon_DebugStepOut_40x", Icon40x40));
+		Set( "PlayWorld.StepOut.Small", new IMAGE_BRUSH("Icons/icon_DebugStepOut_40x", Icon20x20));
+		Set( "PlayWorld.StepInto", new IMAGE_BRUSH( "Icons/icon_DebugStepIn_40x", Icon40x40 ) );
+		Set( "PlayWorld.StepInto.Small", new IMAGE_BRUSH( "Icons/icon_DebugStepIn_40x", Icon20x20 ) );
+		Set( "PlayWorld.StepOver", new IMAGE_BRUSH( "Icons/icon_DebugStepOver_40x", Icon40x40 ) );
+		Set( "PlayWorld.StepOver.Small", new IMAGE_BRUSH("Icons/icon_DebugStepOver_40x", Icon20x20));
 	}
 
 
@@ -5624,12 +5627,13 @@ void FSlateEditorStyle::FStyle::SetupPersonaStyle()
 		Set( "Kismet.DebuggerOverlay.Breakpoint.Disabled", new IMAGE_BRUSH( "Old/Kismet2/Breakpoint_Disabled", Icon32x32 ) );
 		Set( "Kismet.DebuggerOverlay.Breakpoint.EnabledAndInvalid", new IMAGE_BRUSH( "Old/Kismet2/Breakpoint_Invalid", Icon32x32 ) );
 		Set( "Kismet.DebuggerOverlay.Breakpoint.EnabledAndValid", new IMAGE_BRUSH( "Old/Kismet2/Breakpoint_Valid", Icon32x32 ) );
-		Set( "Kismet.DebuggerOverlay.Breakpoint.DisabledCollapsed", new IMAGE_BRUSH( "Old/Kismet2/Breakpoint_Disabled_Collapsed", Icon32x32 ) );
-		Set( "Kismet.DebuggerOverlay.Breakpoint.EnabledAndInvalidCollapsed", new IMAGE_BRUSH( "Old/Kismet2/Breakpoint_Invalid_Collapsed", Icon32x32 ) );
-		Set( "Kismet.DebuggerOverlay.Breakpoint.EnabledAndValidCollapsed", new IMAGE_BRUSH( "Old/Kismet2/Breakpoint_Valid_Collapsed", Icon32x32 ) );
 
 		Set( "Kismet.DebuggerOverlay.InstructionPointer", new IMAGE_BRUSH( "Old/Kismet2/IP_Normal", FVector2D(128,96)) );
 		Set( "Kismet.DebuggerOverlay.InstructionPointerBreakpoint", new IMAGE_BRUSH( "Old/Kismet2/IP_Breakpoint", FVector2D(128,96)) );
+
+		Set ("Kismet.CallStackViewer.CurrentStackFrame", new IMAGE_BRUSH( "Old/Kismet2/DebuggerArrow", Icon12x12 ));
+		Set( "Kismet.CallStackViewer.CurrentStackFrameColor", FLinearColor(0.728f, 0.364f, 0.003f) );
+		Set( "Kismet.CallStackViewer.LastStackFrameNavigatedToColor", FLinearColor( 0.4f, 0.5f, 0.7f, 1.0f ) );
 	}
 
 	// Asset context menu
diff --git a/Engine/Source/Editor/GraphEditor/Private/BlueprintConnectionDrawingPolicy.cpp b/Engine/Source/Editor/GraphEditor/Private/BlueprintConnectionDrawingPolicy.cpp
--- a/Engine/Source/Editor/GraphEditor/Private/BlueprintConnectionDrawingPolicy.cpp
+++ b/Engine/Source/Editor/GraphEditor/Private/BlueprintConnectionDrawingPolicy.cpp
@@ -5,12 +5,62 @@
 #include "Engine/Blueprint.h"
 #include "Engine/BlueprintGeneratedClass.h"
 #include "EdGraphSchema_K2.h"
+#include "K2Node_Composite.h"
 #include "K2Node_Knot.h"
 #include "K2Node_MacroInstance.h"
+#include "K2Node_TunnelBoundary.h"
 #include "Kismet2/KismetDebugUtilities.h"
 #include "Kismet2/BlueprintEditorUtils.h"
 #include "BlueprintEditorSettings.h"
 
+namespace
+{
+	// Attempt to find a tunnel instance node in the given source graph that matches the given target graph.
+	// @TODO - This is a potential performance issue as it's a linear/recursive search. We used to map nested
+	// Macro Instances to code offsets through the debug data, and search only the mapped set, but the solution
+	// only worked for Macro Instance nodes. Consider reusing that solution and extending it to include Composites.
+	static UEdGraphNode* FindMatchingTunnelInstanceNode(const UEdGraph* InGraph, const UEdGraph* TargetGraph)
+	{
+		// Gather up all tunnel nodes in the given source graph.
+		TArray<UK2Node_Tunnel*> TunnelNodes;
+		InGraph->GetNodesOfClass(TunnelNodes);
+		for (UEdGraphNode* TunnelNode : TunnelNodes)
+		{
+			// Check all tunnel instance nodes for a match against the associated tunnel instance graph.
+			if (UK2Node_MacroInstance* MacroInstanceNode = Cast<UK2Node_MacroInstance>(TunnelNode))
+			{
+				const UEdGraph* MacroGraph = MacroInstanceNode->GetMacroGraph();
+				if (TargetGraph == MacroGraph)
+				{
+					// The current macro source graph matches up with the target graph context.
+					return MacroInstanceNode;
+				}
+				else if(UEdGraphNode* TunnelInstanceNode = FindMatchingTunnelInstanceNode(MacroGraph, TargetGraph))
+				{
+					// We recursively found a tunnel instance node nested inside the macro source graph that matches up with the target graph context.
+					return TunnelInstanceNode;
+				}
+			}
+			else if (UK2Node_Composite* CompositeNode = Cast<UK2Node_Composite>(TunnelNode))
+			{
+				if (TargetGraph == CompositeNode->BoundGraph)
+				{
+					// The current composite source graph matches up with the target graph context.
+					return CompositeNode;
+				}
+				else if(UEdGraphNode* TunnelInstanceNode = FindMatchingTunnelInstanceNode(CompositeNode->BoundGraph, TargetGraph))
+				{
+					// We recursively found a tunnel instance node nested inside the composite source graph that matches up with the target graph context.
+					return TunnelInstanceNode;
+				}
+			}
+		}
+
+		// No matching tunnel instance node could be found within the given source graph.
+		return nullptr;
+	}
+};
+
 /////////////////////////////////////////////////////
 // FKismetConnectionDrawingPolicy
 
@@ -121,50 +171,11 @@ void FKismetConnectionDrawingPolicy::BuildExecutionRoadmap()
 							SequentialNodeTimes.Add(Sample.ObservationTime);
 							SequentialExecPinsInGraph.Add(AssociatedPin);
 						}
-						else
+						else if (UEdGraphNode* TunnelInstanceNode = FindMatchingTunnelInstanceNode(GraphObj, Node->GetGraph()))
 						{
-							// If the top-level source node is a macro instance node
-							UK2Node_MacroInstance* MacroInstanceNode = Cast<UK2Node_MacroInstance>(Node);
-							if (MacroInstanceNode)
-							{
-								// Attempt to locate the macro source node through the code mapping
-								UEdGraphNode* MacroSourceNode = DebugData.FindMacroSourceNodeFromCodeLocation(Sample.Function.Get(), Sample.Offset);
-								if (MacroSourceNode)
-								{
-									// If the macro source node is located in the current graph context
-									if (GraphObj == MacroSourceNode->GetGraph())
-									{
-										// Add it to the sequential node list
-										SequentialNodesInGraph.Add(MacroSourceNode);
-										SequentialNodeTimes.Add(Sample.ObservationTime);
-										SequentialExecPinsInGraph.Add(AssociatedPin);
-									}
-									else
-									{
-										// The macro source node isn't in the current graph context, but we might have a macro instance node that is
-										// in the current graph context, so obtain the set of macro instance nodes that are mapped to the code here.
-										TArray<UEdGraphNode*> MacroInstanceNodes;
-										DebugData.FindMacroInstanceNodesFromCodeLocation(Sample.Function.Get(), Sample.Offset, MacroInstanceNodes);
-
-										// For each macro instance node in the set
-										for (auto MacroInstanceNodeIt = MacroInstanceNodes.CreateConstIterator(); MacroInstanceNodeIt; ++MacroInstanceNodeIt)
-										{
-											// If the macro instance node is located in the current graph context
-											MacroInstanceNode = Cast<UK2Node_MacroInstance>(*MacroInstanceNodeIt);
-											if (MacroInstanceNode && GraphObj == MacroInstanceNode->GetGraph())
-											{
-												// Add it to the sequential node list
-												SequentialNodesInGraph.Add(MacroInstanceNode);
-												SequentialNodeTimes.Add(Sample.ObservationTime);
-												SequentialExecPinsInGraph.Add(AssociatedPin);
-
-												// Exit the loop; we're done
-												break;
-											}
-										}
-									}
-								}
-							}
+							SequentialNodesInGraph.Add(TunnelInstanceNode);
+							SequentialNodeTimes.Add(Sample.ObservationTime);
+							SequentialExecPinsInGraph.Add(AssociatedPin);
 						}
 					}
 				}
diff --git a/Engine/Source/Editor/GraphEditor/Private/KismetNodes/SGraphNodeK2Base.cpp b/Engine/Source/Editor/GraphEditor/Private/KismetNodes/SGraphNodeK2Base.cpp
--- a/Engine/Source/Editor/GraphEditor/Private/KismetNodes/SGraphNodeK2Base.cpp
+++ b/Engine/Source/Editor/GraphEditor/Private/KismetNodes/SGraphNodeK2Base.cpp
@@ -15,8 +15,6 @@
 #include "SGraphPin.h"
 #include "EdGraphSchema_K2.h"
 #include "K2Node.h"
-#include "K2Node_Composite.h"
-#include "K2Node_MacroInstance.h"
 #include "K2Node_Timeline.h"
 #include "Engine/Breakpoint.h"
 #include "Kismet2/KismetDebugUtilities.h"
@@ -361,28 +359,13 @@ void SGraphNodeK2Base::GetOverlayBrushes(bool bSelected, const FVector2D WidgetS
 	{
 		FOverlayBrushInfo BreakpointOverlayInfo;
 
-		if (Breakpoint->GetLocation()->IsA<UK2Node_Composite>()
-			|| Breakpoint->GetLocation()->IsA<UK2Node_MacroInstance>())
+		if (Breakpoint->IsEnabledByUser())
 		{
-			if (Breakpoint->IsEnabledByUser())
-			{
-				BreakpointOverlayInfo.Brush = FEditorStyle::GetBrush(FKismetDebugUtilities::IsBreakpointValid(Breakpoint) ? TEXT("Kismet.DebuggerOverlay.Breakpoint.EnabledAndValidCollapsed") : TEXT("Kismet.DebuggerOverlay.Breakpoint.EnabledAndInvalidCollapsed"));
-			}
-			else
-			{
-				BreakpointOverlayInfo.Brush = FEditorStyle::GetBrush(TEXT("Kismet.DebuggerOverlay.Breakpoint.DisabledCollapsed"));
-			}
+			BreakpointOverlayInfo.Brush = FEditorStyle::GetBrush(FKismetDebugUtilities::IsBreakpointValid(Breakpoint) ? TEXT("Kismet.DebuggerOverlay.Breakpoint.EnabledAndValid") : TEXT("Kismet.DebuggerOverlay.Breakpoint.EnabledAndInvalid"));
 		}
 		else
 		{
-			if (Breakpoint->IsEnabledByUser())
-			{
-				BreakpointOverlayInfo.Brush = FEditorStyle::GetBrush(FKismetDebugUtilities::IsBreakpointValid(Breakpoint) ? TEXT("Kismet.DebuggerOverlay.Breakpoint.EnabledAndValid") : TEXT("Kismet.DebuggerOverlay.Breakpoint.EnabledAndInvalid"));
-			}
-			else
-			{
-				BreakpointOverlayInfo.Brush = FEditorStyle::GetBrush(TEXT("Kismet.DebuggerOverlay.Breakpoint.Disabled"));
-			}
+			BreakpointOverlayInfo.Brush = FEditorStyle::GetBrush(TEXT("Kismet.DebuggerOverlay.Breakpoint.Disabled"));
 		}
 
 		if(BreakpointOverlayInfo.Brush != NULL)
diff --git a/Engine/Source/Editor/GraphEditor/Private/KismetNodes/SGraphNodeSpawnActor.cpp b/Engine/Source/Editor/GraphEditor/Private/KismetNodes/SGraphNodeSpawnActor.cpp
--- a/Engine/Source/Editor/GraphEditor/Private/KismetNodes/SGraphNodeSpawnActor.cpp
+++ b/Engine/Source/Editor/GraphEditor/Private/KismetNodes/SGraphNodeSpawnActor.cpp
@@ -9,6 +9,7 @@
 #include "KismetPins/SGraphPinObject.h"
 #include "ClassViewerModule.h"
 #include "ClassViewerFilter.h"
+#include "ScopedTransaction.h"
 
 #define LOCTEXT_NAMESPACE "SGraphNodeSpawnActor"
 
@@ -31,6 +32,9 @@ class SGraphPinActorBasedBlueprintClass : public SGraphPinObject
 			check(InChosenClass->IsChildOf(AActor::StaticClass()));
 			if(const UEdGraphSchema* Schema = GraphPinObj->GetSchema())
 			{
+				const FScopedTransaction Transaction(NSLOCTEXT("GraphEditor", "ChangeClassPinValue", "Change Class Pin Value"));
+				GraphPinObj->Modify();
+
 				Schema->TrySetDefaultObject(*GraphPinObj, InChosenClass->ClassGeneratedBy);
 			}
 		}
diff --git a/Engine/Source/Editor/GraphEditor/Private/KismetNodes/SGraphNodeSpawnActorFromClass.cpp b/Engine/Source/Editor/GraphEditor/Private/KismetNodes/SGraphNodeSpawnActorFromClass.cpp
--- a/Engine/Source/Editor/GraphEditor/Private/KismetNodes/SGraphNodeSpawnActorFromClass.cpp
+++ b/Engine/Source/Editor/GraphEditor/Private/KismetNodes/SGraphNodeSpawnActorFromClass.cpp
@@ -12,6 +12,7 @@
 #include "NodeFactory.h"
 #include "ClassViewerModule.h"
 #include "ClassViewerFilter.h"
+#include "ScopedTransaction.h"
 
 #define LOCTEXT_NAMESPACE "SGraphPinActorBasedClass"
 
@@ -32,6 +33,9 @@ class SGraphPinActorBasedClass : public SGraphPinObject
 		{
 			if(const UEdGraphSchema* Schema = GraphPinObj->GetSchema())
 			{
+				const FScopedTransaction Transaction(NSLOCTEXT("GraphEditor", "ChangeClassPinValue", "Change Class Pin Value"));
+				GraphPinObj->Modify();
+
 				Schema->TrySetDefaultObject(*GraphPinObj, InChosenClass);
 			}
 		}
@@ -75,6 +79,9 @@ protected:
 			const UClass* SelectedClass = GEditor->GetFirstSelectedClass(PinRequiredParentClass);
 			if(SelectedClass)
 			{
+				const FScopedTransaction Transaction(NSLOCTEXT("GraphEditor", "ChangeClassPinValue", "Change Class Pin Value"));
+				GraphPinObj->Modify();
+
 				GraphPinObj->GetSchema()->TrySetDefaultObject(*GraphPinObj, const_cast<UClass*>(SelectedClass));
 			}
 		}
diff --git a/Engine/Source/Editor/GraphEditor/Private/KismetPins/SGraphPinClass.cpp b/Engine/Source/Editor/GraphEditor/Private/KismetPins/SGraphPinClass.cpp
--- a/Engine/Source/Editor/GraphEditor/Private/KismetPins/SGraphPinClass.cpp
+++ b/Engine/Source/Editor/GraphEditor/Private/KismetPins/SGraphPinClass.cpp
@@ -33,6 +33,9 @@ FReply SGraphPinClass::OnClickUse()
 		const UClass* SelectedClass = GEditor->GetFirstSelectedClass(PinRequiredParentClass);
 		if(SelectedClass)
 		{
+			const FScopedTransaction Transaction(NSLOCTEXT("GraphEditor", "ChangeClassPinValue", "Change Class Pin Value"));
+			GraphPinObj->Modify();
+
 			GraphPinObj->GetSchema()->TrySetDefaultObject(*GraphPinObj, const_cast<UClass*>(SelectedClass));
 		}
 	}
diff --git a/Engine/Source/Editor/GraphEditor/Private/KismetPins/SGraphPinObject.cpp b/Engine/Source/Editor/GraphEditor/Private/KismetPins/SGraphPinObject.cpp
--- a/Engine/Source/Editor/GraphEditor/Private/KismetPins/SGraphPinObject.cpp
+++ b/Engine/Source/Editor/GraphEditor/Private/KismetPins/SGraphPinObject.cpp
@@ -152,6 +152,9 @@ FReply SGraphPinObject::OnClickUse()
 		UObject* SelectedObject = GEditor->GetSelectedObjects()->GetTop(ObjectClass);
 		if(SelectedObject != NULL)
 		{
+			const FScopedTransaction Transaction(NSLOCTEXT("GraphEditor", "ChangeObjectPinValue", "Change Object Pin Value"));
+			GraphPinObj->Modify();
+
 			GraphPinObj->GetSchema()->TrySetDefaultObject(*GraphPinObj, SelectedObject);
 		}
 	}
diff --git a/Engine/Source/Editor/Kismet/Private/BlueprintCompilationManager.cpp b/Engine/Source/Editor/Kismet/Private/BlueprintCompilationManager.cpp
--- a/Engine/Source/Editor/Kismet/Private/BlueprintCompilationManager.cpp
+++ b/Engine/Source/Editor/Kismet/Private/BlueprintCompilationManager.cpp
@@ -1355,12 +1355,21 @@ void FBlueprintCompilationManagerImpl::ReinstanceBatch(TArray<FReinstancingJob>&
 					UObject* Iter = Archetype->GetOuter();
 					while(Iter)
 					{
+						UBlueprintGeneratedClass* IterAsBPGC = Cast<UBlueprintGeneratedClass>(Iter);
 						if(Iter->HasAnyFlags(RF_ClassDefaultObject)
-							|| Cast<UBlueprintGeneratedClass>(Iter)
+							|| IterAsBPGC
 							|| Cast<UBlueprint>(Iter) )
 						{
 							ArchetypeReferencers.Add(Iter);
 
+							// Component templates are referenced by the UBlueprint, but are outered to the UBPGC. Both
+							// will need to be updated. Realistically there is no reason to refernce these in the 
+							// UBlueprint, so there is no reason to generalize this behavior:
+							if(IterAsBPGC)
+							{
+								ArchetypeReferencers.Add(IterAsBPGC->ClassGeneratedBy);
+							}
+
 							// this handles nested subobjects:
 							TArray<UObject*> ContainedObjects;
 							GetObjectsWithOuter(Iter, ContainedObjects);
diff --git a/Engine/Source/Editor/Kismet/Private/BlueprintEditor.cpp b/Engine/Source/Editor/Kismet/Private/BlueprintEditor.cpp
--- a/Engine/Source/Editor/Kismet/Private/BlueprintEditor.cpp
+++ b/Engine/Source/Editor/Kismet/Private/BlueprintEditor.cpp
@@ -3438,9 +3438,13 @@ void FBlueprintEditor::JumpToHyperlink(const UObject* ObjectReference, bool bReq
 	}
 	else if(const UFunction* Function = Cast<const UFunction>(ObjectReference))
 	{
-		if (UEdGraph* FunctionGraph = FBlueprintEditorUtils::FindScopeGraph(GetBlueprintObj(), Function))
+		UBlueprint* BP = GetBlueprintObj();
+		if(BP)
 		{
-			OpenDocument(FunctionGraph, FDocumentTracker::OpenNewDocument);
+			if (UEdGraph* FunctionGraph = FBlueprintEditorUtils::FindScopeGraph(BP, Function))
+			{
+				OpenDocument(FunctionGraph, FDocumentTracker::OpenNewDocument);
+			}
 		}
 	}
 	else if(const UBlueprintGeneratedClass* Class = Cast<const UBlueprintGeneratedClass>(ObjectReference))
diff --git a/Engine/Source/Editor/Kismet/Private/BlueprintEditorModule.cpp b/Engine/Source/Editor/Kismet/Private/BlueprintEditorModule.cpp
--- a/Engine/Source/Editor/Kismet/Private/BlueprintEditorModule.cpp
+++ b/Engine/Source/Editor/Kismet/Private/BlueprintEditorModule.cpp
@@ -31,6 +31,7 @@
 #include "BlueprintGraphPanelPinFactory.h"
 #include "Widgets/Docking/SDockTab.h"
 #include "KismetCompiler.h"
+#include "CallStackViewer.h"
 
 #define LOCTEXT_NAMESPACE "BlueprintEditor"
 
@@ -202,6 +203,8 @@ void FBlueprintEditorModule::StartupModule()
 		.SetGroup( MenuStructure.GetDeveloperToolsDebugCategory() )
 		.SetIcon(FSlateIcon(FEditorStyle::GetStyleSetName(), "BlueprintDebugger.TabIcon"));
 
+	CallStackViewer::RegisterTabSpawner();
+
 	// Have to check GIsEditor because right now editor modules can be loaded by the game
 	// Once LoadModule is guaranteed to return NULL for editor modules in game, this can be removed
 	// Without this check, loading the level editor in the game will crash
diff --git a/Engine/Source/Editor/Kismet/Private/BlueprintEditorTabFactories.cpp b/Engine/Source/Editor/Kismet/Private/BlueprintEditorTabFactories.cpp
--- a/Engine/Source/Editor/Kismet/Private/BlueprintEditorTabFactories.cpp
+++ b/Engine/Source/Editor/Kismet/Private/BlueprintEditorTabFactories.cpp
@@ -395,16 +395,23 @@ TSharedRef<SWidget> FFindResultsSummoner::CreateTabBody(const FWorkflowTabSpawnI
 	return BlueprintEditorPtr->GetFindResults();
 }
 
-void FGraphTabHistory::EvokeHistory(TSharedPtr<FTabInfo> InTabInfo)
+void FGraphTabHistory::EvokeHistory(TSharedPtr<FTabInfo> InTabInfo, bool bPrevTabMatches)
 {
 	FWorkflowTabSpawnInfo SpawnInfo;
 	SpawnInfo.Payload = Payload;
 	SpawnInfo.TabInfo = InTabInfo;
 
-	TSharedRef< SGraphEditor > GraphEditorRef = StaticCastSharedRef< SGraphEditor >(FactoryPtr.Pin()->CreateTabBody(SpawnInfo));
-	GraphEditor = GraphEditorRef;
-
-	FactoryPtr.Pin()->UpdateTab(InTabInfo->GetTab().Pin(), SpawnInfo, GraphEditorRef);
+	if(bPrevTabMatches)
+	{
+		TSharedPtr<SDockTab> DockTab = InTabInfo->GetTab().Pin();
+		GraphEditor = StaticCastSharedRef<SGraphEditor>(DockTab->GetContent());
+	}
+	else
+	{
+		TSharedRef< SGraphEditor > GraphEditorRef = StaticCastSharedRef< SGraphEditor >(FactoryPtr.Pin()->CreateTabBody(SpawnInfo));
+		GraphEditor = GraphEditorRef;
+		FactoryPtr.Pin()->UpdateTab(InTabInfo->GetTab().Pin(), SpawnInfo, GraphEditorRef);
+	}
 }
 
 void FGraphTabHistory::SaveHistory()
diff --git a/Engine/Source/Editor/Kismet/Private/BlueprintEditorTabFactories.h b/Engine/Source/Editor/Kismet/Private/BlueprintEditorTabFactories.h
--- a/Engine/Source/Editor/Kismet/Private/BlueprintEditorTabFactories.h
+++ b/Engine/Source/Editor/Kismet/Private/BlueprintEditorTabFactories.h
@@ -105,7 +105,7 @@ public:
 
 	}
 
-	virtual void EvokeHistory(TSharedPtr<FTabInfo> InTabInfo) override;
+	virtual void EvokeHistory(TSharedPtr<FTabInfo> InTabInfo, bool bPrevTabMatches) override;
 
 	virtual void SaveHistory() override;
 
diff --git a/Engine/Source/Editor/Kismet/Private/FindInBlueprintManager.cpp b/Engine/Source/Editor/Kismet/Private/FindInBlueprintManager.cpp
--- a/Engine/Source/Editor/Kismet/Private/FindInBlueprintManager.cpp
+++ b/Engine/Source/Editor/Kismet/Private/FindInBlueprintManager.cpp
@@ -979,7 +979,7 @@ namespace BlueprintSearchMetaDataHelpers
 								InWriter->WriteObjectStart();
 								{
 									InWriter->WriteValue(FFindInBlueprintSearchTags::FiB_Name, Pin->GetSchema()->GetPinDisplayName(Pin));
-									InWriter->WriteValue(FFindInBlueprintSearchTags::FiB_DefaultValue, FText::FromString(Pin->GetDefaultAsString()));
+									InWriter->WriteValue(FFindInBlueprintSearchTags::FiB_DefaultValue, Pin->GetDefaultAsText());
 								}
 								SavePinTypeToJson(InWriter, Pin->PinType);
 								InWriter->WriteObjectEnd();
diff --git a/Engine/Source/Editor/Kismet/Private/WorkflowOrientedApp/WorkflowTabManager.cpp b/Engine/Source/Editor/Kismet/Private/WorkflowOrientedApp/WorkflowTabManager.cpp
--- a/Engine/Source/Editor/Kismet/Private/WorkflowOrientedApp/WorkflowTabManager.cpp
+++ b/Engine/Source/Editor/Kismet/Private/WorkflowOrientedApp/WorkflowTabManager.cpp
@@ -28,13 +28,16 @@ FTabInfo::FTabInfo(const TSharedRef<SDockTab>& InTab, const TSharedPtr<FDocument
 
 bool FTabInfo::PayloadMatches(const TSharedPtr<FTabPayload> TestPayload) const
 {
-	TSharedPtr<FTabPayload> CurrentPayload = History[CurrentHistoryIndex]->GetPayload();
+	return PayloadMatches(History[CurrentHistoryIndex]->GetPayload(), TestPayload);
+}
 
-	if (CurrentPayload.IsValid() && TestPayload.IsValid())
+bool FTabInfo::PayloadMatches(TSharedPtr<FTabPayload> A, TSharedPtr<FTabPayload> B)
+{
+	if (A.IsValid() && B.IsValid())
 	{
-		return CurrentPayload->IsEqual(TestPayload.ToSharedRef());
+		return A->IsEqual(B.ToSharedRef());
 	}
-	else if (!TestPayload.IsValid() && !CurrentPayload.IsValid())
+	else if (!B.IsValid() && !A.IsValid())
 	{
 		return true;
 	}
@@ -44,7 +47,7 @@ bool FTabInfo::PayloadMatches(const TSharedPtr<FTabPayload> TestPayload) const
 	}
 }
 
-void FTabInfo::AddTabHistory(TSharedPtr< struct FGenericTabHistory > InHistoryNode, bool bInSaveHistory/* = true*/)
+void FTabInfo::AddTabHistory(TSharedPtr< struct FGenericTabHistory > InHistoryNode, bool bInSaveHistory/* = true*/, bool bSpawnNewTab/* = true*/)
 {
 	// If the tab is not new, save the current history.
 	if(CurrentHistoryIndex >= 0 && bInSaveHistory)
@@ -72,7 +75,7 @@ void FTabInfo::AddTabHistory(TSharedPtr< struct FGenericTabHistory > InHistoryNo
 	CurrentHistoryIndex = History.Num() - 1;
 
 	// Evoke the history
-	InHistoryNode->EvokeHistory(AsShared());
+	InHistoryNode->EvokeHistory(AsShared(), bSpawnNewTab);
 	InHistoryNode->GetFactory().Pin()->OnTabActivated(Tab.Pin());
 }
 
@@ -80,6 +83,7 @@ FReply FTabInfo::OnGoForwardInHistory()
 {
 	if( CurrentHistoryIndex < History.Num() - 1 )
 	{
+		const int32 PreviousHistoryIndex = CurrentHistoryIndex;
 		History[CurrentHistoryIndex]->SaveHistory();
 
 		while( CurrentHistoryIndex < History.Num() - 1)
@@ -88,7 +92,8 @@ FReply FTabInfo::OnGoForwardInHistory()
 
 			if( History[CurrentHistoryIndex]->IsHistoryValid() )
 			{
-				History[CurrentHistoryIndex]->EvokeHistory(AsShared());
+				bool bPayloadMatches = PayloadMatches(History[PreviousHistoryIndex]->GetPayload(), History[CurrentHistoryIndex]->GetPayload());
+				History[CurrentHistoryIndex]->EvokeHistory(AsShared(), bPayloadMatches);
 				History[CurrentHistoryIndex]->RestoreHistory();
 				History[CurrentHistoryIndex]->GetFactory().Pin()->OnTabActivated(Tab.Pin());
 				break;
@@ -102,6 +107,7 @@ FReply FTabInfo::OnGoBackInHistory()
 {
 	if( CurrentHistoryIndex > 0 )
 	{
+		const int32 PreviousHistoryIndex = CurrentHistoryIndex;
 		History[CurrentHistoryIndex]->SaveHistory();
 
 		while( CurrentHistoryIndex > 0 )
@@ -110,7 +116,8 @@ FReply FTabInfo::OnGoBackInHistory()
 				
 			if( History[CurrentHistoryIndex]->IsHistoryValid() )
 			{
-				History[CurrentHistoryIndex]->EvokeHistory(AsShared());
+				bool bPayloadMatches = PayloadMatches(History[PreviousHistoryIndex]->GetPayload(), History[CurrentHistoryIndex]->GetPayload());
+				History[CurrentHistoryIndex]->EvokeHistory(AsShared(), bPayloadMatches);
 				History[CurrentHistoryIndex]->RestoreHistory();
 				History[CurrentHistoryIndex]->GetFactory().Pin()->OnTabActivated(Tab.Pin());
 
@@ -145,7 +152,7 @@ void FTabInfo::JumpToNearestValidHistoryData()
 				}
 			}
 
-			History[CurrentHistoryIndex]->EvokeHistory(AsShared());
+			History[CurrentHistoryIndex]->EvokeHistory(AsShared(), false);
 			History[CurrentHistoryIndex]->RestoreHistory();
 			History[CurrentHistoryIndex]->GetFactory().Pin()->OnTabActivated(Tab.Pin());
 			FGlobalTabmanager::Get()->SetActiveTab(nullptr);
@@ -168,10 +175,13 @@ void FTabInfo::GoToHistoryIndex(int32 InHistoryIdx)
 {
 	if(History[InHistoryIdx]->IsHistoryValid())
 	{
+		const int32 PreviousHistoryIndex = CurrentHistoryIndex;
 		History[CurrentHistoryIndex]->SaveHistory();
 
 		CurrentHistoryIndex = InHistoryIdx;
-		History[CurrentHistoryIndex]->EvokeHistory(AsShared());
+		
+		bool bPayloadMatches = PayloadMatches(History[PreviousHistoryIndex]->GetPayload(), History[CurrentHistoryIndex]->GetPayload());
+		History[CurrentHistoryIndex]->EvokeHistory(AsShared(), bPayloadMatches);
 		History[CurrentHistoryIndex]->RestoreHistory();
 		History[CurrentHistoryIndex]->GetFactory().Pin()->OnTabActivated(Tab.Pin());
 	}
@@ -489,9 +499,10 @@ TSharedPtr<SDockTab> FDocumentTracker::OpenDocument(TSharedPtr<FTabPayload> InPa
 		TSharedPtr<SDockTab> Tab;
 
 		// If the current tab matches we'll re-use it.
-		if(LastEditedTabInfo.IsValid() && LastEditedTabInfo.Pin()->PayloadMatches(InPayload))
+		TSharedPtr<FTabInfo> LastEditedTabInfoPinned = LastEditedTabInfo.Pin();
+		if(LastEditedTabInfoPinned.IsValid() && LastEditedTabInfoPinned->PayloadMatches(InPayload))
 		{
-			Tab = LastEditedTabInfo.Pin()->GetTab().Pin();
+			Tab = LastEditedTabInfoPinned->GetTab().Pin();
 		}
 		else
 		{
@@ -514,8 +525,8 @@ TSharedPtr<SDockTab> FDocumentTracker::OpenDocument(TSharedPtr<FTabPayload> InPa
 		}
 		else if(InOpenCause == CreateHistoryEvent)
 		{
-			// We are forcing a history event, navigate the current tab.
-			Tab = NavigateCurrentTab(InPayload, NavigatingCurrentDocument);
+			TSharedPtr<FDocumentTabFactory> Factory = FindSupportingFactory(InPayload.ToSharedRef());
+			LastEditedTabInfoPinned->AddTabHistory(Factory->CreateTabHistoryNode(InPayload), true, true);
 		}
 		
 		return Tab;
diff --git a/Engine/Source/Editor/Kismet/Public/CallStackViewer.h b/Engine/Source/Editor/Kismet/Public/CallStackViewer.h
new file mode 100644
--- /dev/null
+++ b/Engine/Source/Editor/Kismet/Public/CallStackViewer.h
@@ -0,0 +1,12 @@
+// Copyright 1998-2017 Epic Games, Inc. All Rights Reserved.
+
+#pragma once
+
+#include "Widgets/SCompoundWidget.h"
+#include "WorkflowOrientedApp/WorkflowTabFactory.h"
+
+namespace CallStackViewer
+{
+	void KISMET_API UpdateDisplayedCallstack(const TArray<const FFrame*>& ScriptStack);
+	void RegisterTabSpawner();
+}
diff --git a/Engine/Source/Editor/Kismet/Public/WorkflowOrientedApp/WorkflowTabManager.h b/Engine/Source/Editor/Kismet/Public/WorkflowOrientedApp/WorkflowTabManager.h
--- a/Engine/Source/Editor/Kismet/Public/WorkflowOrientedApp/WorkflowTabManager.h
+++ b/Engine/Source/Editor/Kismet/Public/WorkflowOrientedApp/WorkflowTabManager.h
@@ -65,13 +65,15 @@ public:
 	/** Returns TRUE if the payload used by this TabInfo is the same as passed in */
 	bool PayloadMatches(const TSharedPtr<FTabPayload> TestPayload) const;
 
+	static bool PayloadMatches(TSharedPtr<FTabPayload> A, TSharedPtr<FTabPayload> B);
+
 	/**
 	 * Adds history data to immediately after the current location in the history list, wiping out any history after in the process
 	 *
 	 * @param InHistoryNode		The history node to add
 	 * @param bInSaveHistory	TRUE if history should be saved
 	 */
-	void AddTabHistory(TSharedPtr< struct FGenericTabHistory > InHistoryNode, bool bInSaveHistory = true);
+	void AddTabHistory(TSharedPtr< struct FGenericTabHistory > InHistoryNode, bool bInSaveHistory = true, bool bPrevTabMatches = false);
 
 	/** Single step forward in history */
 	FReply OnGoForwardInHistory();
@@ -158,13 +160,16 @@ public:
 	 *
 	 * @param InTab				The tab to update with the history data
 	 */
-	virtual void EvokeHistory(TSharedPtr<FTabInfo> InTabInfo)
+	virtual void EvokeHistory(TSharedPtr<FTabInfo> InTabInfo, bool bPrevTabMatches)
 	{
 		FWorkflowTabSpawnInfo SpawnInfo;
 		SpawnInfo.Payload = Payload;
 		SpawnInfo.TabInfo = InTabInfo;
 
-		FactoryPtr.Pin()->UpdateTab(InTabInfo->GetTab().Pin(), SpawnInfo, FactoryPtr.Pin()->CreateTabBody(SpawnInfo));
+		if(!bPrevTabMatches)
+		{	
+			FactoryPtr.Pin()->UpdateTab(InTabInfo->GetTab().Pin(), SpawnInfo, FactoryPtr.Pin()->CreateTabBody(SpawnInfo));
+		}
 	}
 
 	/** Saves any important information from the payload into the history node */
diff --git a/Engine/Source/Editor/KismetCompiler/Private/KismetCompiler.cpp b/Engine/Source/Editor/KismetCompiler/Private/KismetCompiler.cpp
--- a/Engine/Source/Editor/KismetCompiler/Private/KismetCompiler.cpp
+++ b/Engine/Source/Editor/KismetCompiler/Private/KismetCompiler.cpp
@@ -24,20 +24,20 @@
 #include "Engine/UserDefinedStruct.h"
 #include "EdGraphUtilities.h"
 #include "K2Node_CallFunction.h"
-#include "K2Node_Knot.h"
-#include "K2Node_Tunnel.h"
 #include "K2Node_Composite.h"
 #include "K2Node_CreateDelegate.h"
 #include "K2Node_CustomEvent.h"
 #include "K2Node_FunctionEntry.h"
 #include "K2Node_FunctionResult.h"
+#include "K2Node_Knot.h"
 #include "K2Node_MacroInstance.h"
 #include "K2Node_MakeArray.h"
 #include "K2Node_TemporaryVariable.h"
 #include "K2Node_Timeline.h"
+#include "K2Node_Tunnel.h"
+#include "K2Node_TunnelBoundary.h"
 #include "K2Node_VariableGet.h"
 #include "K2Node_VariableSet.h"
-#include "K2Node_TunnelBoundary.h"
 #include "KismetCompilerBackend.h"
 #include "Kismet2/KismetReinstanceUtilities.h"
 #include "Engine/SCS_Node.h"
@@ -2086,10 +2086,21 @@ void FKismetCompilerContext::SetCalculatedMetaDataAndFlags(UFunction* Function,
 		if(	!K2Schema->IsMetaPin(*EntryPin) && 
 			(EntryPin->PinType.PinCategory != UEdGraphSchema_K2::PC_Object) && 
 			(EntryPin->PinType.PinCategory != UEdGraphSchema_K2::PC_Class) && 
-			(EntryPin->PinType.PinCategory != UEdGraphSchema_K2::PC_Interface) && 
-			!EntryPin->DefaultValue.IsEmpty() )
+			(EntryPin->PinType.PinCategory != UEdGraphSchema_K2::PC_Interface))
 		{
-			Function->SetMetaData(EntryPin->PinName, *EntryPin->DefaultValue);
+			if (!EntryPin->DefaultValue.IsEmpty())
+			{
+				Function->SetMetaData(EntryPin->PinName, *EntryPin->DefaultValue);
+			}
+			else if (!EntryPin->DefaultTextValue.IsEmpty())
+			{
+				FString TextAsString;
+				if (FTextStringHelper::WriteToString(TextAsString, EntryPin->DefaultTextValue))
+				{
+					Function->SetMetaData(EntryPin->PinName, *TextAsString);
+				}
+			}
+			
 		}
 	}
 
@@ -2645,9 +2656,6 @@ void FKismetCompilerContext::CreateFunctionStubForEvent(UK2Node_Event* SrcEventN
 	FKismetFunctionContext& StubContext = *new (FunctionList) FKismetFunctionContext(MessageLog, Schema, NewClass, Blueprint, CompileOptions.DoesRequireCppCodeGeneration());
 	StubContext.SourceGraph = ChildStubGraph;
 
-	// A stub graph has no visual representation and is thus not suited to be debugged via the debugger
-	StubContext.bCreateDebugData = false;
-
 	StubContext.SourceEventFromStubGraph = SrcEventNode;
 
 	if (SrcEventNode->bOverrideFunction || SrcEventNode->bInternalEvent)
@@ -3225,16 +3233,6 @@ void FKismetCompilerContext::ExpandTunnelsAndMacros(UEdGraph* SourceGraph)
 			// Clone the macro graph, then move all of its children, keeping a list of nodes from the macro
 			UEdGraph* ClonedGraph = FEdGraphUtilities::CloneGraph(MacroGraph, NULL, &MessageLog, true);
 
-			// Find active tunnel instances for this macro instance.
-			TArray<TWeakObjectPtr<UEdGraphNode>> TunnelInstances;
-			MessageLog.GetTunnelsActiveForNode(MacroInstanceNode, TunnelInstances);
-			if (!TunnelInstances.Num())
-			{
-				UEdGraphNode* SourceMacroInstance = Cast<UEdGraphNode>(MessageLog.FindSourceObject(MacroInstanceNode));
-				TunnelInstances.Add(SourceMacroInstance);
-				MessageLog.RegisterIntermediateTunnelInstance(MacroInstanceNode, TunnelInstances);
-			}
-
 			for (int32 I = 0; I < ClonedGraph->Nodes.Num(); ++I)
 			{
 				MacroGeneratedNodes.Add(ClonedGraph->Nodes[I], CurrentNode);
@@ -3348,13 +3346,6 @@ void FKismetCompilerContext::ExpandTunnelsAndMacros(UEdGraph* SourceGraph)
 						DuplicatedNode->ReconstructNode();
 					}
 
-					// Record the source node mapping for the intermediate node and register any active tunnel instances if it is a tunnel/macro instance node.
-					MessageLog.RegisterIntermediateTunnelNode(DuplicatedNode, MacroInstanceNode);
-					if (FBlueprintEditorUtils::IsTunnelInstanceNode(DuplicatedNode))
-					{
-						MessageLog.RegisterIntermediateTunnelInstance(DuplicatedNode, TunnelInstances);
-					}
-
 					DuplicatedNode->NodePosY += NodeOffsetY;
 					DuplicatedNode->NodePosX += NodeOffsetX;
 
@@ -3391,29 +3382,21 @@ void FKismetCompilerContext::ExpandTunnelsAndMacros(UEdGraph* SourceGraph)
 		}
 		else if (TunnelNode)
 		{
-			UEdGraphNode* InputSink = TunnelNode->GetInputSink();
-			for (UEdGraphPin const* TunnelPin : TunnelNode->Pins)
-			{
-				if ((TunnelPin->Direction != EGPD_Input) || (TunnelPin->PinType.PinCategory != UEdGraphSchema_K2::PC_Exec))
-				{
-					continue;
-				}
-				check(InputSink != NULL);
+			UK2Node_Tunnel* InputSink = TunnelNode->GetInputSink();
+			UK2Node_Tunnel* OutputSource = TunnelNode->GetOutputSource();
 
-				UEdGraphPin* SinkPin = InputSink->FindPin(TunnelPin->PinName);
-				if (SinkPin == NULL)
-				{
-					continue;
-				}
-				check(SinkPin->Direction == EGPD_Output);
-
-				for (UEdGraphPin* TunnelLinkedPin : TunnelPin->LinkedTo)
-				{
-					MessageLog.NotifyIntermediatePinCreation(TunnelLinkedPin, SinkPin);
-				}
+			if (FBlueprintEditorUtils::IsTunnelInstanceNode(OutputSource))
+			{
+				// Process the input side of a tunnel instance expansion
+				ProcessIntermediateTunnelBoundary(OutputSource, TunnelNode);
+			}
+			else if (FBlueprintEditorUtils::IsTunnelInstanceNode(InputSink))
+			{
+				// Process the output side of a tunnel instance expansion
+				ProcessIntermediateTunnelBoundary(TunnelNode, InputSink);
 			}
 
-			bool bSuccess = Schema->CollapseGatewayNode(TunnelNode, InputSink, TunnelNode->GetOutputSource(), this);
+			bool bSuccess = Schema->CollapseGatewayNode(TunnelNode, InputSink, OutputSource, this);
 			if (!bSuccess)
 			{
 				MessageLog.Error(*LOCTEXT("CollapseTunnel_Error", "Failed to collapse tunnel @@").ToString(), TunnelNode);
@@ -4503,5 +4486,131 @@ void FKismetCompilerContext::RegisterCompilerForBP(UClass* BPClass, TFunction<TS
 	CustomCompilerMap.Add(BPClass, FactoryFunction);
 }
 
+void FKismetCompilerContext::MapExpansionPathToTunnelInstance(UEdGraphNode* InnerExpansionNode, UEdGraphNode* OuterTunnelInstance)
+{
+	if (InnerExpansionNode && OuterTunnelInstance)
+	{
+		MessageLog.NotifyIntermediateTunnelNode(InnerExpansionNode, OuterTunnelInstance);
+
+		// Recursively map any nodes linked to this node along the execution path.
+		for (const UEdGraphPin* OutputPin : InnerExpansionNode->Pins)
+		{
+			if (OutputPin->Direction == EGPD_Output && UEdGraphSchema_K2::IsExecPin(*OutputPin) && OutputPin->LinkedTo.Num() > 0)
+			{
+				for (const UEdGraphPin* LinkedTo : OutputPin->LinkedTo)
+				{
+					// Make sure it hasn't already been mapped (e.g. re-entrant execution paths).
+					UEdGraphNode* ExpansionNode = LinkedTo->GetOwningNode();
+					UEdGraphNode* TunnelInstance = MessageLog.GetIntermediateTunnelInstance(ExpansionNode);
+					if (TunnelInstance != OuterTunnelInstance)
+					{
+						MapExpansionPathToTunnelInstance(ExpansionNode, OuterTunnelInstance);
+					}
+				}
+			}
+		}
+	}
+}
+
+void FKismetCompilerContext::ProcessIntermediateTunnelBoundary(UK2Node_Tunnel* TunnelInput, UK2Node_Tunnel* TunnelOutput)
+{
+	// @TODO move this check out of KismetFunctionContext so we can use it here?
+	auto IsDebuggingOrInstrumentationRequired = []() -> bool
+	{
+		return GIsEditor && !IsRunningCommandlet();
+	};
+
+	// Common initialization.
+	auto InitializeTunnelBoundaryNode = [this](UK2Node_TunnelBoundary* TunnelBoundary, UK2Node_Tunnel* TunnelSource)
+	{
+		// Set the base node name and boundary type.
+		TunnelBoundary->SetNodeAttributes(TunnelSource);
+
+		// Auto-position the node in the intermediate graph.
+		// @TODO - Maybe we need to find a better way to reposition these?
+		AutoAssignNodePosition(TunnelBoundary);
+	};
+
+	if (TunnelInput)
+	{
+		const bool bIsInputBoundary = FBlueprintEditorUtils::IsTunnelInstanceNode(TunnelInput);
+
+		for (UEdGraphPin* InputPin : TunnelInput->Pins)
+		{
+			// We create a boundary node for each exec pin input. This way every execution path has a debug site.
+			if (InputPin->Direction == EGPD_Input && UEdGraphSchema_K2::IsExecPin(*InputPin) && InputPin->LinkedTo.Num() > 0)
+			{
+				if (IsDebuggingOrInstrumentationRequired())
+				{
+					// Create one or more boundary nodes that precede the tunnel input node.
+					if (UK2Node_TunnelBoundary* InputBoundaryNode = SpawnIntermediateNode<UK2Node_TunnelBoundary>(TunnelInput))
+					{
+						InitializeTunnelBoundaryNode(InputBoundaryNode, TunnelInput);
+
+						if (UEdGraphPin* NewInputPin = InputBoundaryNode->CreatePin(EGPD_Input, InputPin->PinType, InputPin->PinName))
+						{
+							if (UEdGraphPin* NewOutputPin = InputBoundaryNode->CreatePin(EGPD_Output, InputPin->PinType, InputBoundaryNode->CreateUniquePinName(InputPin->PinName)))
+							{
+								// Move the exec pin links to the boundary node. This ensures that execution will flow through the boundary node.
+								if (MovePinLinksToIntermediate(*InputPin, *NewInputPin).CanSafeConnect())
+								{
+									NewOutputPin->MakeLinkTo(InputPin);
+								}
+							}
+						}
+					}
+				}
+
+				// Look for a matching pin on the tunnel output node.
+				if (UEdGraphPin* OutputPin = TunnelOutput ? TunnelOutput->FindPin(InputPin->PinName) : nullptr)
+				{
+					if (ensure(OutputPin->Direction == EGPD_Output && UEdGraphSchema_K2::IsExecPin(*OutputPin)) && OutputPin->LinkedTo.Num() > 0)
+					{
+						if (bIsInputBoundary)
+						{
+							// Map the execution path through the expansion back to the tunnel instance node. Note that the assumption here is
+							// that we haven't collapsed the tunnels yet, so the output side of the expansion shouldn't be linked to anything.
+							for (UEdGraphPin* LinkedTo : OutputPin->LinkedTo)
+							{
+								MapExpansionPathToTunnelInstance(LinkedTo->GetOwningNode(), TunnelInput);
+							}
+
+							if (IsDebuggingOrInstrumentationRequired())
+							{
+								// We also create a boundary node for each matching exec pin on the tunnel output node.
+								if (UK2Node_TunnelBoundary* OutputBoundaryNode = SpawnIntermediateNode<UK2Node_TunnelBoundary>(TunnelOutput))
+								{
+									InitializeTunnelBoundaryNode(OutputBoundaryNode, TunnelOutput);
+
+									if (UEdGraphPin* NewInputPin = OutputBoundaryNode->CreatePin(EGPD_Input, OutputPin->PinType, OutputPin->PinName))
+									{
+										if (UEdGraphPin* NewOutputPin = OutputBoundaryNode->CreatePin(EGPD_Output, OutputPin->PinType, OutputBoundaryNode->CreateUniquePinName(OutputPin->PinName)))
+										{
+											// Move the exec pin links to the boundary node. This ensures that execution will flow through the boundary node.
+											if (MovePinLinksToIntermediate(*OutputPin, *NewOutputPin).CanSafeConnect())
+											{
+												NewInputPin->MakeLinkTo(OutputPin);
+											}
+										}
+									}
+								}
+							}
+						}
+						else if(IsDebuggingOrInstrumentationRequired())
+						{
+							// This is the output side of the expansion, so a tunnel boundary node will not be required on the output side of the pair. However, for
+							// wire traces to function properly, we still need to map exec pins linked to the input side back to the matching pin on the output side.
+							for (UEdGraphPin* LinkedInputPin : InputPin->LinkedTo)
+							{
+								MessageLog.NotifyIntermediatePinCreation(LinkedInputPin, OutputPin);
+							}
+						}
+					}
+				}
+			}
+		}
+	}
+}
+
 #undef LOCTEXT_NAMESPACE
 //////////////////////////////////////////////////////////////////////////
diff --git a/Engine/Source/Editor/KismetCompiler/Private/KismetCompilerMisc.cpp b/Engine/Source/Editor/KismetCompiler/Private/KismetCompilerMisc.cpp
--- a/Engine/Source/Editor/KismetCompiler/Private/KismetCompilerMisc.cpp
+++ b/Engine/Source/Editor/KismetCompiler/Private/KismetCompilerMisc.cpp
@@ -2141,7 +2141,7 @@ struct FEventGraphUtils
 		}
 
 		// If the function call node is an intermediate node resulting from expansion of an async task node, then the return value term must also be persistent.
-		const UEdGraphNode* SourceNode = MessageLog.GetSourceNode(OwnerNode);
+		const UObject* SourceNode = MessageLog.FindSourceObject(OwnerNode);
 		if (SourceNode && SourceNode->IsA<UK2Node_BaseAsyncTask>())
 		{
 			return true;
diff --git a/Engine/Source/Editor/KismetCompiler/Private/KismetCompilerVMBackend.cpp b/Engine/Source/Editor/KismetCompiler/Private/KismetCompilerVMBackend.cpp
--- a/Engine/Source/Editor/KismetCompiler/Private/KismetCompilerVMBackend.cpp
+++ b/Engine/Source/Editor/KismetCompiler/Private/KismetCompilerVMBackend.cpp
@@ -1853,46 +1853,14 @@ public:
 		if (SourceNode != NULL)
 		{
 			// Record where this NOP is
-			UEdGraphNode* TrueSourceNode = FunctionContext.MessageLog.GetSourceNode(SourceNode);
+			UEdGraphNode* TrueSourceNode = Cast<UEdGraphNode>(FunctionContext.MessageLog.FindSourceObject(SourceNode));
 			if (TrueSourceNode)
 			{
-				// If this is a debug site for an expanded macro instruction, there should also be a macro source node associated with it
-				UEdGraphNode* MacroSourceNode = CompilerContext.MessageLog.GetSourceTunnelNode(SourceNode);
-				// We need to ensure that macro/composite instances also record the tunnels present at any script location ( including themselves )
-				if (MacroSourceNode == TrueSourceNode && !FBlueprintEditorUtils::IsTunnelInstanceNode(MacroSourceNode))
-				{
-					// The function above will return the given node if not found in the map. In that case there is no associated source macro node, so we clear it.
-					MacroSourceNode = NULL;
-				}
-
-				TArray<TWeakObjectPtr<UEdGraphNode>> MacroInstanceNodes;
 				const bool bInstrumentedBreakpoint = Statement.Type == KCST_InstrumentedWireEntry;
-				bool bBreakpointSite = Statement.Type == KCST_DebugSite || bInstrumentedBreakpoint;
-
-				if (MacroSourceNode)
-				{
-					// Only associate macro instance node breakpoints with source nodes that are linked to the entry node in an impure macro graph
-					if (bBreakpointSite)
-					{
-						const UK2Node_MacroInstance* MacroInstanceNode = Cast<const UK2Node_MacroInstance>(TrueSourceNode);
-						if (MacroInstanceNode)
-						{
-							TArray<const UEdGraphNode*> ValidBreakpointLocations;
-							FKismetDebugUtilities::GetValidBreakpointLocations(MacroInstanceNode, ValidBreakpointLocations);
-							bBreakpointSite = ValidBreakpointLocations.Contains(MacroSourceNode);
-						}
-					}
+				const bool bBreakpointSite = Statement.Type == KCST_DebugSite || bInstrumentedBreakpoint;
 
-					// Gather up all the macro instance nodes that lead to this macro source node
-					UEdGraphNode* IntermediateMacroInstance = CompilerContext.MessageLog.GetIntermediateTunnelInstance(SourceNode);
-					CompilerContext.MessageLog.GetTunnelsActiveForNode(IntermediateMacroInstance, MacroInstanceNodes);
-					if (MacroInstanceNodes.Num())
-					{
-						TrueSourceNode = MacroInstanceNodes[0].Get();
-					}
-				}
 				// Register the debug information for the node.
-				ClassBeingBuilt->GetDebugData().RegisterNodeToCodeAssociation(TrueSourceNode, MacroSourceNode, MacroInstanceNodes, FunctionContext.Function, Offset, bBreakpointSite);
+				ClassBeingBuilt->GetDebugData().RegisterNodeToCodeAssociation(TrueSourceNode, FunctionContext.Function, Offset, bBreakpointSite);
 
 				// Track pure node script code range for the current impure (exec) node
 				if (Statement.Type == KCST_InstrumentedPureNodeEntry)
@@ -1908,7 +1876,7 @@ public:
 				else if (Statement.Type == KCST_InstrumentedWireEntry && PureNodeEntryCount > 0)
 				{
 					// Map script code range for the full set of pure node inputs feeding in to the current impure (exec) node at the current offset
-					ClassBeingBuilt->GetDebugData().RegisterPureNodeScriptCodeRange(MacroSourceNode ? MacroSourceNode : TrueSourceNode, FunctionContext.Function, FInt32Range(PureNodeEntryStart, Offset));
+					ClassBeingBuilt->GetDebugData().RegisterPureNodeScriptCodeRange(TrueSourceNode, FunctionContext.Function, FInt32Range(PureNodeEntryStart, Offset));
 
 					// Reset pure node code range tracking.
 					PureNodeEntryCount = 0;
diff --git a/Engine/Source/Editor/KismetCompiler/Public/KismetCompiler.h b/Engine/Source/Editor/KismetCompiler/Public/KismetCompiler.h
--- a/Engine/Source/Editor/KismetCompiler/Public/KismetCompiler.h
+++ b/Engine/Source/Editor/KismetCompiler/Public/KismetCompiler.h
@@ -18,6 +18,7 @@ class UBlueprintGeneratedClass;
 class UK2Node_FunctionEntry;
 class UK2Node_TemporaryVariable;
 class UK2Node_Timeline;
+class UK2Node_Tunnel;
 class FKismetCompilerVMBackend;
 
 KISMETCOMPILER_API DECLARE_LOG_CATEGORY_EXTERN(LogK2Compiler, Log, All);
@@ -435,6 +436,27 @@ protected:
 	 */
 	void ExpandTunnelsAndMacros(UEdGraph* SourceGraph);
 
+	/**
+	 * Maps the nodes in an intermediate tunnel expansion path back to the owning tunnel instance node.
+	 */
+	void MapExpansionPathToTunnelInstance(UEdGraphNode* InnerExpansionNode, UEdGraphNode* OuterTunnelInstance);
+
+	/**
+	* Processes an intermediate tunnel expansion boundary.
+	*
+	* We define a tunnel boundary as the input and output sides of an intermediate tunnel instance node expansion. Each boundary
+	* consists of a pair of tunnel nodes (input/output), with one side being the tunnel "instance" node that owns the expansion.
+	* After expansion, tunnel nodes are cropped and removed from the function graph, so they do not result in any actual bytecode.
+	*
+	* This function maps the nodes in the execution path through the expansion and back to the outer tunnel instance node. If
+	* Blueprint debugging is enabled, this function also spawns one or more intermediate "boundary" NOPs around the tunnel I/O
+	* pair. The boundary nodes are intended to serve as debug sites, allowing breakpoints to be hit on both sides of the tunnel.
+	*
+	* @param	TunnelInput		Tunnel input node. This will either be a tunnel instance node (OutputSource) or a tunnel exit node.
+	* @param	TunnelOutput	Tunnel output node. This will either be a tunnel entry node or a tunnel instance node (InputSink).
+	*/
+	void ProcessIntermediateTunnelBoundary(UK2Node_Tunnel* TunnelInput, UK2Node_Tunnel* TunnelOutput);
+
 	/**
 	 * Merges pages and creates function stubs, etc...
 	 */
diff --git a/Engine/Source/Editor/PropertyEditor/Private/UserInterface/PropertyEditor/SPropertyEditorAsset.cpp b/Engine/Source/Editor/PropertyEditor/Private/UserInterface/PropertyEditor/SPropertyEditorAsset.cpp
--- a/Engine/Source/Editor/PropertyEditor/Private/UserInterface/PropertyEditor/SPropertyEditorAsset.cpp
+++ b/Engine/Source/Editor/PropertyEditor/Private/UserInterface/PropertyEditor/SPropertyEditorAsset.cpp
@@ -819,11 +819,22 @@ FPropertyAccess::Result SPropertyEditorAsset::GetValue( FObjectOrAssetData& OutV
 	}
 	else
 	{
+		FSoftObjectPath SoftObjectPath;
 		UObject* Object = nullptr;
 		if (PropertyHandle.IsValid())
 		{
 			Result = PropertyHandle->GetValue(Object);
 		}
+		else
+		{
+			SoftObjectPath = FSoftObjectPath(ObjectPath.Get());
+			Object = SoftObjectPath.ResolveObject();
+
+			if (Object != nullptr)
+			{
+				Result = FPropertyAccess::Success;
+			}
+		}
 
 		if (Object != nullptr)
 		{
@@ -839,13 +850,14 @@ FPropertyAccess::Result SPropertyEditorAsset::GetValue( FObjectOrAssetData& OutV
 		}
 		else
 		{
-			const FString CurrentObjectPath = ObjectPath.Get();
-			Result = FPropertyAccess::Success;
-
-			FSoftObjectPath SoftObjectPath = FSoftObjectPath(CurrentObjectPath);
+			if (SoftObjectPath.IsNull())
+			{
+				SoftObjectPath = FSoftObjectPath(ObjectPath.Get());
+			}
 
 			if (SoftObjectPath.IsAsset())
 			{
+				const FString CurrentObjectPath = SoftObjectPath.ToString();
 				if (CurrentObjectPath != TEXT("None") && (!CachedAssetData.IsValid() || CachedAssetData.ObjectPath.ToString() != CurrentObjectPath))
 				{
 					static FName AssetRegistryName("AssetRegistry");
@@ -855,6 +867,7 @@ FPropertyAccess::Result SPropertyEditorAsset::GetValue( FObjectOrAssetData& OutV
 				}
 
 				OutValue = FObjectOrAssetData(CachedAssetData);
+				Result = FPropertyAccess::Success;
 			}
 			else
 			{
diff --git a/Engine/Source/Editor/StatsViewer/Private/StatsPages/PrimitiveStatsPage.cpp b/Engine/Source/Editor/StatsViewer/Private/StatsPages/PrimitiveStatsPage.cpp
--- a/Engine/Source/Editor/StatsViewer/Private/StatsPages/PrimitiveStatsPage.cpp
+++ b/Engine/Source/Editor/StatsViewer/Private/StatsPages/PrimitiveStatsPage.cpp
@@ -279,7 +279,7 @@ struct PrimitiveStatsGenerator
 				NewStatsEntry->HWInstances		= NumHWInstances;
 				NewStatsEntry->Triangles		= 0;
 				NewStatsEntry->InstTriangles	= 0;
-				NewStatsEntry->ResourceSize		= (float)(FArchiveCountMem(Resource).GetNum() + Resource->GetResourceSizeBytes(EResourceSizeMode::Exclusive)) / 1024.0f;
+				NewStatsEntry->ResourceSize		= (float)(FArchiveCountMem(Resource).GetNum() + Resource->GetResourceSizeBytes(EResourceSizeMode::EstimatedTotal)) / 1024.0f;
 				NewStatsEntry->Sections			= 0;
 				NewStatsEntry->InstSections		= 0;
 				NewStatsEntry->RadiusMin		= InPrimitiveComponent->Bounds.SphereRadius;
@@ -356,7 +356,7 @@ struct PrimitiveStatsGenerator
 						UniqueTextures.Add(CurrentComponent->HeightmapTexture, &bNotUnique);
 						if (!bNotUnique)
 						{
-							const SIZE_T HeightmapResourceSize = CurrentComponent->HeightmapTexture->GetResourceSizeBytes(EResourceSizeMode::Exclusive);
+							const SIZE_T HeightmapResourceSize = CurrentComponent->HeightmapTexture->GetResourceSizeBytes(EResourceSizeMode::EstimatedTotal);
 							NewStatsEntry->ResourceSize += HeightmapResourceSize;
 						}
 						if (CurrentComponent->XYOffsetmapTexture)
@@ -364,7 +364,7 @@ struct PrimitiveStatsGenerator
 							UniqueTextures.Add(CurrentComponent->XYOffsetmapTexture, &bNotUnique);
 							if (!bNotUnique)
 							{
-								const SIZE_T OffsetmapResourceSize = CurrentComponent->XYOffsetmapTexture->GetResourceSizeBytes(EResourceSizeMode::Exclusive);
+								const SIZE_T OffsetmapResourceSize = CurrentComponent->XYOffsetmapTexture->GetResourceSizeBytes(EResourceSizeMode::EstimatedTotal);
 								NewStatsEntry->ResourceSize += OffsetmapResourceSize;
 							}
 						}
@@ -374,7 +374,7 @@ struct PrimitiveStatsGenerator
 							UniqueTextures.Add((*ItWeightmaps), &bNotUnique);
 							if (!bNotUnique)
 							{
-								const SIZE_T WeightmapResourceSize = (*ItWeightmaps)->GetResourceSizeBytes(EResourceSizeMode::Exclusive);
+								const SIZE_T WeightmapResourceSize = (*ItWeightmaps)->GetResourceSizeBytes(EResourceSizeMode::EstimatedTotal);
 								NewStatsEntry->ResourceSize += WeightmapResourceSize;
 							}
 						}
diff --git a/Engine/Source/Editor/UnrealEd/Private/Commandlets/AssetRegistryGenerator.cpp b/Engine/Source/Editor/UnrealEd/Private/Commandlets/AssetRegistryGenerator.cpp
--- a/Engine/Source/Editor/UnrealEd/Private/Commandlets/AssetRegistryGenerator.cpp
+++ b/Engine/Source/Editor/UnrealEd/Private/Commandlets/AssetRegistryGenerator.cpp
@@ -730,7 +730,7 @@ bool FAssetRegistryGenerator::SaveAssetRegistry(const FString& SandboxPath, bool
 
 		// Save the generated registry
 		FString PlatformSandboxPath = SandboxPath.Replace(TEXT("[Platform]"), *TargetPlatform->PlatformName());
-		PlatformSandboxPath.ReplaceInline(TEXT("AssetRegistry.bin"), TEXT("DevelopmentAssetRegistry.bin"));
+		PlatformSandboxPath.ReplaceInline(TEXT("AssetRegistry.bin"), TEXT("Metadata/DevelopmentAssetRegistry.bin"));
 		FFileHelper::SaveArrayToFile(SerializedAssetRegistry, *PlatformSandboxPath);
 
 		if (bGenerateChunks && bUseAssetManager)
diff --git a/Engine/Source/Editor/UnrealEd/Private/Commandlets/PackageUtilities.cpp b/Engine/Source/Editor/UnrealEd/Private/Commandlets/PackageUtilities.cpp
--- a/Engine/Source/Editor/UnrealEd/Private/Commandlets/PackageUtilities.cpp
+++ b/Engine/Source/Editor/UnrealEd/Private/Commandlets/PackageUtilities.cpp
@@ -2095,9 +2095,7 @@ struct CompressAnimationsFunctor
 			SIZE_T OldSize;
 			SIZE_T NewSize;
 
-			{
-				OldSize = AnimSeq->GetResourceSizeBytes(EResourceSizeMode::Inclusive);
-			}
+			OldSize = AnimSeq->GetResourceSizeBytes(EResourceSizeMode::EstimatedTotal);
 
 			// Clear bDoNotOverrideCompression flag
 			if( bClearNoCompressionOverride && AnimSeq->bDoNotOverrideCompression )
@@ -2151,9 +2149,7 @@ struct CompressAnimationsFunctor
 				bDirtyPackage = true;
 			}
 
-			{
-				NewSize = AnimSeq->GetResourceSizeBytes(EResourceSizeMode::Inclusive);
-			}
+			NewSize = AnimSeq->GetResourceSizeBytes(EResourceSizeMode::EstimatedTotal);
 
 			// Only save package if size has changed.
 			const int64 DeltaSize = NewSize - OldSize;
diff --git a/Engine/Source/Editor/UnrealEd/Private/ComponentTypeRegistry.cpp b/Engine/Source/Editor/UnrealEd/Private/ComponentTypeRegistry.cpp
--- a/Engine/Source/Editor/UnrealEd/Private/ComponentTypeRegistry.cpp
+++ b/Engine/Source/Editor/UnrealEd/Private/ComponentTypeRegistry.cpp
@@ -330,7 +330,7 @@ void FComponentTypeRegistryData::ForceRefreshComponentList()
 			// GetAssetsByClass call is a kludge to get the full asset paths for the blueprints we care about, Bob T. thinks 
 			// that the Asset Registry could just keep asset paths:
 			TArray<FAssetData> BlueprintAssetData;
-			AssetRegistryModule.Get().GetAssetsByClass(UBlueprint::StaticClass()->GetFName(), BlueprintAssetData);
+			AssetRegistryModule.Get().GetAssetsByClass(UBlueprint::StaticClass()->GetFName(), BlueprintAssetData, true);
 			TMap<FString, FAssetData> BlueprintNames;
 			for (const FAssetData& Blueprint : BlueprintAssetData)
 			{
diff --git a/Engine/Source/Editor/UnrealEd/Private/CookOnTheFlyServer.cpp b/Engine/Source/Editor/UnrealEd/Private/CookOnTheFlyServer.cpp
--- a/Engine/Source/Editor/UnrealEd/Private/CookOnTheFlyServer.cpp
+++ b/Engine/Source/Editor/UnrealEd/Private/CookOnTheFlyServer.cpp
@@ -3851,7 +3851,7 @@ bool UCookOnTheFlyServer::GetCurrentIniVersionStrings( const ITargetPlatform* Ta
 
 bool UCookOnTheFlyServer::GetCookedIniVersionStrings(const ITargetPlatform* TargetPlatform, FIniSettingContainer& OutIniSettings, TMap<FString,FString>& OutAdditionalSettings) const
 {
-	const FString EditorIni = FPaths::ProjectDir() / TEXT("CookedIniVersion.txt");
+	const FString EditorIni = FPaths::ProjectDir() / TEXT("Metadata") / TEXT("CookedIniVersion.txt");
 	const FString SandboxEditorIni = ConvertToFullSandboxPath(*EditorIni, true);
 
 
@@ -4185,7 +4185,7 @@ bool UCookOnTheFlyServer::SaveCurrentIniSettings(const ITargetPlatform* TargetPl
 	FIniSettingContainer CurrentIniSettings;
 	GetCurrentIniVersionStrings(TargetPlatform, CurrentIniSettings);
 
-	const FString EditorIni = FPaths::ProjectDir() / TEXT("CookedIniVersion.txt");
+	const FString EditorIni = FPaths::ProjectDir() / TEXT("Metadata") / TEXT("CookedIniVersion.txt");
 	const FString SandboxEditorIni = ConvertToFullSandboxPath(*EditorIni, true);
 
 
@@ -4355,7 +4355,7 @@ void UCookOnTheFlyServer::PopulateCookedPackagesFromDisk(const TArray<ITargetPla
 		FAssetRegistryGenerator* PlatformAssetRegistry = RegistryGenerators.FindRef(PlatformFName);
 
 		// Load the platform cooked asset registry file
-		const FString CookedAssetRegistry = FPaths::ProjectDir() / GetDevelopmentAssetRegistryFilename();
+		const FString CookedAssetRegistry = FPaths::ProjectDir() / TEXT("Metadata") / GetDevelopmentAssetRegistryFilename();
 		const FString SandboxCookedAssetRegistryFilename = ConvertToFullSandboxPath(*CookedAssetRegistry, true, Target->PlatformName());
 
 		bool bIsIterateSharedBuild = IsCookFlagSet(ECookInitializationFlags::IterateSharedBuild);
@@ -4366,12 +4366,10 @@ void UCookOnTheFlyServer::PopulateCookedPackagesFromDisk(const TArray<ITargetPla
 			FDateTime CurrentLocalCookedBuild = IFileManager::Get().GetTimeStamp(*SandboxCookedAssetRegistryFilename);
 
 			// iterate on the shared build if the option is set
-			FString SharedCookedAssetRegistry = FPaths::Combine(*FPaths::ProjectSavedDir(), TEXT("SharedIterativeBuild"), *Target->PlatformName(), TEXT("Cooked"), GetDevelopmentAssetRegistryFilename());
+			FString SharedCookedAssetRegistry = FPaths::Combine(*FPaths::ProjectSavedDir(), TEXT("SharedIterativeBuild"), *Target->PlatformName(), TEXT("Metadata"), GetDevelopmentAssetRegistryFilename());
 
 			FDateTime CurrentIterativeCookedBuild = IFileManager::Get().GetTimeStamp(*SharedCookedAssetRegistry);
 
-
-
 			if ( (CurrentIterativeCookedBuild >= CurrentLocalCookedBuild) && 
 				(CurrentIterativeCookedBuild != FDateTime::MinValue()) )
 			{
@@ -4382,8 +4380,8 @@ void UCookOnTheFlyServer::PopulateCookedPackagesFromDisk(const TArray<ITargetPla
 
 				// SaveCurrentIniSettings(Target); // use this if we don't care about ini safty.
 				// copy the ini settings from the shared cooked build. 
-				const FString SharedCookedIniFile = FPaths::Combine(*FPaths::ProjectSavedDir(), TEXT("SharedIterativeBuild"), *Target->PlatformName(), TEXT("Cooked"), TEXT("CookedIniVersion.txt"));
-				const FString SandboxCookedIniFile = ConvertToFullSandboxPath(*(FPaths::ProjectDir() / TEXT("CookedIniVersion.txt")), true).Replace(TEXT("[Platform]"), *Target->PlatformName());
+				const FString SharedCookedIniFile = FPaths::Combine(*FPaths::ProjectSavedDir(), TEXT("SharedIterativeBuild"), *Target->PlatformName(), TEXT("Metadata"), TEXT("CookedIniVersion.txt"));
+				const FString SandboxCookedIniFile = ConvertToFullSandboxPath(*(FPaths::ProjectDir() / TEXT("Metadata") / TEXT("CookedIniVersion.txt")), true).Replace(TEXT("[Platform]"), *Target->PlatformName());
 
 				IFileManager::Get().Copy(*SandboxCookedIniFile, *SharedCookedIniFile);
 
@@ -5541,8 +5539,8 @@ void UCookOnTheFlyServer::CookByTheBookFinished()
 					IFileManager::Get().Copy(*VersionedRegistryFilename, *CookedAssetRegistryFilename, true, true);
 
 					// Also copy development registry if it exists
-					const FString DevVersionedRegistryFilename = VersionedRegistryFilename.Replace(TEXT("AssetRegistry.bin"), TEXT("DevelopmentAssetRegistry.bin"));
-					const FString DevCookedAssetRegistryFilename = CookedAssetRegistryFilename.Replace(TEXT("AssetRegistry.bin"), TEXT("DevelopmentAssetRegistry.bin"));
+					const FString DevVersionedRegistryFilename = VersionedRegistryFilename.Replace(TEXT("AssetRegistry.bin"), TEXT("Metadata/DevelopmentAssetRegistry.bin"));
+					const FString DevCookedAssetRegistryFilename = CookedAssetRegistryFilename.Replace(TEXT("AssetRegistry.bin"), TEXT("Metadata/DevelopmentAssetRegistry.bin"));
 					IFileManager::Get().Copy(*DevVersionedRegistryFilename, *DevCookedAssetRegistryFilename, true, true);
 				}
 			}
diff --git a/Engine/Source/Editor/UnrealEd/Private/EditorEngine.cpp b/Engine/Source/Editor/UnrealEd/Private/EditorEngine.cpp
--- a/Engine/Source/Editor/UnrealEd/Private/EditorEngine.cpp
+++ b/Engine/Source/Editor/UnrealEd/Private/EditorEngine.cpp
@@ -687,7 +687,7 @@ void UEditorEngine::HandlePackageReloaded(const EPackageReloadPhase InPackageRel
 				// Notify about any BP assets that are about to be unloaded
 				if (UBlueprint* BP = Cast<UBlueprint>(InObject))
 				{
-					FKismetEditorUtilities::OnBlueprintUnloaded.Broadcast(BP);
+					BP->ClearEditorReferences();
 				}
 			}
 		}, false, RF_Transient, EInternalObjectFlags::PendingKill);
diff --git a/Engine/Source/Editor/UnrealEd/Private/Kismet2/BlueprintEditorUtils.cpp b/Engine/Source/Editor/UnrealEd/Private/Kismet2/BlueprintEditorUtils.cpp
--- a/Engine/Source/Editor/UnrealEd/Private/Kismet2/BlueprintEditorUtils.cpp
+++ b/Engine/Source/Editor/UnrealEd/Private/Kismet2/BlueprintEditorUtils.cpp
@@ -2471,6 +2471,12 @@ void FBlueprintEditorUtils::AddUbergraphPage(UBlueprint* Blueprint, class UEdGra
 #endif	//#if WITH_EDITORONLY_DATA
 }
 
+FName FBlueprintEditorUtils::GetUbergraphFunctionName(const UBlueprint* ForBlueprint)
+{
+	const FString UbergraphCallString = UEdGraphSchema_K2::FN_ExecuteUbergraphBase.ToString() + TEXT("_") + ForBlueprint->GetName();
+	return FName(*UbergraphCallString);
+}
+
 void FBlueprintEditorUtils::AddDomainSpecificGraph(UBlueprint* Blueprint, class UEdGraph* Graph)
 {
 	// Give the schema a chance to fill out any required nodes (like the entry node or results node)
@@ -2842,7 +2848,6 @@ UK2Node_Event* FBlueprintEditorUtils::FindCustomEventNode(const UBlueprint* Blue
 	return FoundNode;
 }
 
-
 void FBlueprintEditorUtils::GatherDependencies(const UBlueprint* InBlueprint, TSet<TWeakObjectPtr<UBlueprint>>& Dependencies, TSet<TWeakObjectPtr<UStruct>>& OutUDSDependencies)
 {
 	struct FGatherDependenciesHelper
@@ -4969,6 +4974,26 @@ UEdGraph* FBlueprintEditorUtils::FindScopeGraph(const UBlueprint* InBlueprint, c
 			break;
 		}
 	}
+
+	if(!ScopeGraph)
+	{
+		FName UbergraphName = FBlueprintEditorUtils::GetUbergraphFunctionName(InBlueprint);
+		
+		bool bIsBlueprintEvent = false;
+		if (const UFunction* AsFunction = Cast<UFunction>(InScope))
+		{
+			bIsBlueprintEvent = AsFunction->HasAnyFunctionFlags(FUNC_BlueprintEvent);
+		}
+
+		if(InScope->GetFName() == UbergraphName || bIsBlueprintEvent)
+		{
+			if(InBlueprint->UbergraphPages.Num() > 0)
+			{
+				ScopeGraph = InBlueprint->UbergraphPages[0];
+			}
+		}
+	}
+
 	return ScopeGraph;
 }
 
diff --git a/Engine/Source/Editor/UnrealEd/Private/Kismet2/CompilerResultsLog.cpp b/Engine/Source/Editor/UnrealEd/Private/Kismet2/CompilerResultsLog.cpp
--- a/Engine/Source/Editor/UnrealEd/Private/Kismet2/CompilerResultsLog.cpp
+++ b/Engine/Source/Editor/UnrealEd/Private/Kismet2/CompilerResultsLog.cpp
@@ -11,7 +11,6 @@
 #include "SourceCodeNavigation.h"
 #include "Developer/HotReload/Public/IHotReload.h"
 #include "EngineLogs.h"
-#include "K2Node_MacroInstance.h"
 #include "Engine/Blueprint.h"
 #include "IMessageLogListing.h"
 
@@ -240,13 +239,6 @@ void FCompilerResultsLog::InternalLogSummary()
 /** Update the source backtrack map to note that NewObject was most closely generated/caused by the SourceObject */
 void FCompilerResultsLog::NotifyIntermediateObjectCreation(UObject* NewObject, UObject* SourceObject)
 {
-	if(UEdGraphNode* NewNode = Cast<UEdGraphNode>(NewObject))
-	{
-		if(UEdGraphNode* OldNode = Cast<UEdGraphNode>(SourceObject))
-		{
-			FullSourceBacktrackMap.Add(NewNode, OldNode);
-		}
-	}
 	SourceBacktrackMap.NotifyIntermediateObjectCreation(NewObject, SourceObject);
 }
 
@@ -266,6 +258,21 @@ UObject const* FCompilerResultsLog::FindSourceObject(UObject const* PossiblyDupl
 	return SourceBacktrackMap.FindSourceObject(PossiblyDuplicatedObject);
 }
 
+void FCompilerResultsLog::NotifyIntermediateTunnelNode(UEdGraphNode* Node, UEdGraphNode* OuterTunnelInstance)
+{
+	IntermediateTunnelNodeToTunnelInstanceMap.Add(Node, OuterTunnelInstance);
+}
+
+UEdGraphNode* FCompilerResultsLog::GetIntermediateTunnelInstance(const UEdGraphNode* IntermediateNode) const
+{
+	TWeakObjectPtr<UEdGraphNode> Result;
+	if (const TWeakObjectPtr<UEdGraphNode>* IntermediateTunnelInstanceNode = IntermediateTunnelNodeToTunnelInstanceMap.Find(IntermediateNode))
+	{
+		Result = *IntermediateTunnelInstanceNode;
+	}
+	return Result.Get();
+}
+
 int32 FCompilerResultsLog::CalculateStableIdentifierForLatentActionManager( const UEdGraphNode* Node )
 {
 	/* 
@@ -278,55 +285,40 @@ int32 FCompilerResultsLog::CalculateStableIdentifierForLatentActionManager( cons
 
 		Meeting these two requirements has proved difficult. The edge cases involve
 		macros and nodes that implement UK2Node::ExpandNode, e.g. LoadAsset/LoadAssetClass
-		nodes in Macros. In order to handle that case we use the full source backtrack map.
+		nodes in Macros. In order to handle that case we use the source backtrack map.
 		Typically an intermediate node has a dynamic GUID, which is useless, but source 
-		nodes that came from macros have stable GUIDs, and can be used 
+		nodes that came from tunnel expansions have stable GUIDs, and can be used 
 	*/
 
+	int32 LatentUUID = 0;
 	const UEdGraphNode* OriginalNode = Node;
+	const UEdGraphNode* OuterTunnelInstance = GetIntermediateTunnelInstance(Node);
 
 	// first search for a node with a stable GUID (e.g., not a node that was created via SpawnIntermediateNode, 
 	// but we want to include nodes that are created via macro instantiation):
-	bool bNodeHasStableGUID = false;
-	while(!bNodeHasStableGUID)
+	if (!OuterTunnelInstance)
 	{
-		if(!Node->HasAnyFlags(RF_Transient) || 
-			GetIntermediateTunnelInstance(Node))
+		Node = Cast<UEdGraphNode>(SourceBacktrackMap.FindSourceObject(Node));
+		if (Node && Node->HasAnyFlags(RF_Transient))
 		{
-			bNodeHasStableGUID = true;
-		}
-		else
-		{
-			UEdGraphNode* const* PreviousNode = FullSourceBacktrackMap.Find(Node);
-			if(PreviousNode && *PreviousNode)
-			{
-				Node = *PreviousNode;
-			}
-			else
-			{
-				// we failed to find a source node, bail
-				bNodeHasStableGUID = true;
-				Node = nullptr;
-			}
+			// we failed to find a source node, bail
+			Node = nullptr;
 		}
 	}
-
-	int32 LatentUUID = 0;
 	
 	if(Node)
 	{
 		LatentUUID = GetTypeHash(Node->NodeGuid);
 
 		const UEdGraphNode* ResultNode = Node;
-		const UEdGraphNode* SourceNode = GetIntermediateTunnelInstance(Node);
-		while (SourceNode && SourceNode != ResultNode)
+		while (OuterTunnelInstance && OuterTunnelInstance != ResultNode)
 		{
-			if (SourceNode->NodeGuid.IsValid())
+			if (OuterTunnelInstance->NodeGuid.IsValid())
 			{
-				LatentUUID = HashCombine(LatentUUID, GetTypeHash(SourceNode->NodeGuid));
+				LatentUUID = HashCombine(LatentUUID, GetTypeHash(OuterTunnelInstance->NodeGuid));
 			}
-			ResultNode = SourceNode;
-			SourceNode = GetIntermediateTunnelInstance(ResultNode);
+			ResultNode = OuterTunnelInstance;
+			OuterTunnelInstance = GetIntermediateTunnelInstance(ResultNode);
 		}
 	}
 	else
@@ -681,94 +673,6 @@ void FCompilerResultsLog::AddChildEvent(TSharedPtr<FCompilerEvent>& ParentEventS
 	}
 }
 
-void FCompilerResultsLog::RegisterIntermediateTunnelNode(UEdGraphNode* IntermediateNode, UEdGraphNode* OwningTunnelInstance)
-{
-	IntermediateTunnelNodeToTunnelInstanceMap.Add(IntermediateNode, OwningTunnelInstance);
-	UObject* TrueSourceNode = SourceBacktrackMap.FindSourceObject(IntermediateNode);
-	if (IntermediateNode != TrueSourceNode)
-	{
-		FinalNodeBackToTunnelSourceMap.NotifyIntermediateObjectCreation(IntermediateNode, TrueSourceNode);
-	}
-}
-
-void FCompilerResultsLog::RegisterIntermediateTunnelInstance(UEdGraphNode* IntermediateTunnel, TArray<TWeakObjectPtr<UEdGraphNode>>& ActiveTunnels)
-{
-	for (auto Tunnel : ActiveTunnels)
-	{
-		if (Tunnel.Get()->IsA<UK2Node_MacroInstance>())
-		{
-			IntermediateTunnelInstanceHierarchyMap.Add(IntermediateTunnel, Tunnel);
-		}
-	}
-	UEdGraphNode* TrueSourceTunnelInstance = Cast<UEdGraphNode>(SourceBacktrackMap.FindSourceObject(IntermediateTunnel));
-	if (!ActiveTunnels.Contains(TrueSourceTunnelInstance))
-	{
-		if (TrueSourceTunnelInstance->IsA<UK2Node_MacroInstance>())
-		{
-			IntermediateTunnelInstanceHierarchyMap.Add(IntermediateTunnel, TrueSourceTunnelInstance);
-		}
-	}
-}
-
-UEdGraphNode* FCompilerResultsLog::GetSourceNode(const UEdGraphNode* IntermediateNode)
-{
-	UEdGraphNode* Result = nullptr;
-	if (UEdGraphNode* SourceTunnelInstance = GetSourceTunnelInstance(IntermediateNode))
-	{
-		if (SourceTunnelInstance->IsA<UK2Node_MacroInstance>())
-		{
-			Result = SourceTunnelInstance;
-		}
-	}
-	if (!Result)
-	{
-		Result = Cast<UEdGraphNode>(SourceBacktrackMap.FindSourceObject(const_cast<UEdGraphNode*>(IntermediateNode)));
-	}
-	return Result;
-}
-
-UEdGraphNode* FCompilerResultsLog::GetIntermediateTunnelInstance(const UEdGraphNode* IntermediateNode) const
-{
-	TWeakObjectPtr<UEdGraphNode> Result;
-	if (const TWeakObjectPtr<UEdGraphNode>* IntermediateTunnelInstanceNode = IntermediateTunnelNodeToTunnelInstanceMap.Find(IntermediateNode))
-	{
-		Result = *IntermediateTunnelInstanceNode;
-	}
-	return Result.Get();
-}
-
-UEdGraphNode* FCompilerResultsLog::GetSourceTunnelNode(const UEdGraphNode* IntermediateNode)
-{
-	UEdGraphNode* SourceNode = Cast<UEdGraphNode>(FinalNodeBackToTunnelSourceMap.FindSourceObject(const_cast<UEdGraphNode*>(IntermediateNode)));
-	if (SourceNode == IntermediateNode)
-	{
-		// No source node available because we haven't mapped this yet, this could be a top level Tunnel so look it up.
-		SourceNode = Cast<UEdGraphNode>(SourceBacktrackMap.FindSourceObject(const_cast<UEdGraphNode*>(IntermediateNode)));
-	}
-	return SourceNode;
-}
-
-UEdGraphNode* FCompilerResultsLog::GetSourceTunnelInstance(const UEdGraphNode* IntermediateNode)
-{
-	TWeakObjectPtr<UEdGraphNode> Result;
-	if (UEdGraphNode* IntermediateTunnelInstance = GetIntermediateTunnelInstance(IntermediateNode))
-	{
-		if (UEdGraphNode* SourceTunnelInstanceNode = Cast<UEdGraphNode>(SourceBacktrackMap.FindSourceObject(IntermediateTunnelInstance)))
-		{
-			if (SourceTunnelInstanceNode != IntermediateNode)
-			{
-				Result = SourceTunnelInstanceNode;
-			}
-		}
-	}
-	return Result.Get();
-}
-
-void FCompilerResultsLog::GetTunnelsActiveForNode(const UEdGraphNode* IntermediateNode, TArray<TWeakObjectPtr<UEdGraphNode>>& ActiveTunnelsOut)
-{
-	IntermediateTunnelInstanceHierarchyMap.MultiFind(IntermediateNode, ActiveTunnelsOut, true);
-}
-
 static FName GetBlueprintMessageLogName(UBlueprint* InBlueprint)
 {
 	FName LogListingName;
diff --git a/Engine/Source/Editor/UnrealEd/Private/Kismet2/DebuggerCommands.cpp b/Engine/Source/Editor/UnrealEd/Private/Kismet2/DebuggerCommands.cpp
--- a/Engine/Source/Editor/UnrealEd/Private/Kismet2/DebuggerCommands.cpp
+++ b/Engine/Source/Editor/UnrealEd/Private/Kismet2/DebuggerCommands.cpp
@@ -329,9 +329,9 @@ void FPlayWorldCommands::RegisterCommands()
 	UI_COMMAND( TogglePlayPauseOfPlaySession, "Toggle Play/Pause", "Resume playing if paused, or pause if playing", EUserInterfaceActionType::Button, FInputChord( EKeys::Pause ) );
 	UI_COMMAND( PossessEjectPlayer, "Possess or Eject Player", "Possesses or ejects the player from the camera", EUserInterfaceActionType::Button, FInputChord( EKeys::F8 ) );
 	UI_COMMAND( ShowCurrentStatement, "Find Node", "Show the current node", EUserInterfaceActionType::Button, FInputChord() );
-	UI_COMMAND( StepInto, "Step Into", "Step Into the next node to be executed", EUserInterfaceActionType::Button, FInputChord( EKeys::F10) );
-	UI_COMMAND( StepOver, "Step Over", "Step to the next node to be executed in the current graph", EUserInterfaceActionType::Button, FInputChord() );
-	UI_COMMAND( StepOut, "Step Out", "Step Out to the next node to be executed in the parent graph", EUserInterfaceActionType::Button, FInputChord() );
+	UI_COMMAND( StepInto, "Step Into", "Step Into the next node to be executed", EUserInterfaceActionType::Button, FInputChord(EKeys::F11) );
+	UI_COMMAND( StepOver, "Step Over", "Step to the next node to be executed in the current graph", EUserInterfaceActionType::Button, FInputChord(EKeys::F10) );
+	UI_COMMAND( StepOut, "Step Out", "Step Out to the next node to be executed in the parent graph", EUserInterfaceActionType::Button, FInputChord(EModifierKey::Alt|EModifierKey::Shift, EKeys::F11) );
 
 	// Launch
 	UI_COMMAND( RepeatLastLaunch, "Launch", "Launches the game on the device as the last session launched from the dropdown next to the Play on Device button on the level editor toolbar", EUserInterfaceActionType::Button, FInputChord( EKeys::P, EModifierKey::Alt | EModifierKey::Shift ) )
@@ -669,6 +669,8 @@ void FPlayWorldCommands::BuildToolbar( FToolBarBuilder& ToolbarBuilder, bool bIn
 	// Single-stepping only buttons
 	ToolbarBuilder.AddToolBarButton(FPlayWorldCommands::Get().ShowCurrentStatement, NAME_None, TAttribute<FText>(), TAttribute<FText>(), TAttribute<FSlateIcon>(), FName(TEXT("ShowCurrentStatement")));
 	ToolbarBuilder.AddToolBarButton(FPlayWorldCommands::Get().StepInto, NAME_None, TAttribute<FText>(), TAttribute<FText>(), TAttribute<FSlateIcon>(), FName(TEXT("StepInto")));
+	ToolbarBuilder.AddToolBarButton(FPlayWorldCommands::Get().StepOver, NAME_None, TAttribute<FText>(), TAttribute<FText>(), TAttribute<FSlateIcon>(), FName(TEXT("StepOver")));
+	ToolbarBuilder.AddToolBarButton(FPlayWorldCommands::Get().StepOut, NAME_None, TAttribute<FText>(), TAttribute<FText>(), TAttribute<FSlateIcon>(), FName(TEXT("StepOut")));
 }
 
 static void MakePreviewDeviceMenu(FMenuBuilder& MenuBuilder )
diff --git a/Engine/Source/Editor/UnrealEd/Private/Kismet2/KismetDebugUtilities.cpp b/Engine/Source/Editor/UnrealEd/Private/Kismet2/KismetDebugUtilities.cpp
--- a/Engine/Source/Editor/UnrealEd/Private/Kismet2/KismetDebugUtilities.cpp
+++ b/Engine/Source/Editor/UnrealEd/Private/Kismet2/KismetDebugUtilities.cpp
@@ -18,6 +18,7 @@
 #include "EdGraph/EdGraph.h"
 #include "Editor/UnrealEdEngine.h"
 #include "Settings/EditorExperimentalSettings.h"
+#include "CallStackViewer.h"
 #include "Animation/AnimBlueprintGeneratedClass.h"
 #include "UnrealEdGlobals.h"
 #include "Engine/Breakpoint.h"
@@ -286,62 +287,6 @@ void FKismetDebugUtilities::OnScriptException(const UObject* ActiveObject, const
 				{
 					CheckBreakConditions(NodeStoppedAt, bShouldBreakExecution);
 				}
-				// If the code which we hit was generated by a macro node expansion
-				UK2Node_MacroInstance* MacroInstanceNode = Cast<UK2Node_MacroInstance>(NodeStoppedAt);
-				if (MacroInstanceNode)
-				{
-					// Attempt to find the associated macro source node
-					UEdGraphNode* MacroSourceNode = FindMacroSourceNodeForCodeLocation(ActiveObject, StackFrame.Node, BreakpointOffset);
-					if (MacroSourceNode)
-					{
-						// If the macro source graph is valid
-						UEdGraph* MacroSourceGraph = MacroSourceNode->GetGraph();
-						if (MacroSourceGraph)
-						{
-							// If the macro source blueprint is valid
-							UBlueprint* MacroBlueprint = MacroSourceGraph->GetTypedOuter<UBlueprint>();
-							if (MacroBlueprint)
-							{
-								// If we're not already going to break execution
-								if (!bShouldBreakExecution && !GEditor->ShouldEndPlayMap())
-								{
-									// Check the source graph to see if any breakpoints are set in the actual macro
-									UBreakpoint* MacroBreakpoint = FindBreakpointForNode(MacroBlueprint, MacroSourceNode, true);
-									if (MacroBreakpoint && MacroBreakpoint->IsEnabledByUser())
-									{
-										// Found one; break at this node in the macro source graph
-										bShouldBreakExecution = true;
-
-										// Redirect breakpoint exception info
-										static FBlueprintExceptionInfo MacroBreakpointExceptionInfo(EBlueprintExceptionType::Breakpoint);
-										ExceptionInfo = &MacroBreakpointExceptionInfo;
-									}
-								}
-
-								// If we're going to break execution on a macro source node
-								if (bShouldBreakExecution)
-								{
-									// Switch to the macro source node
-									NodeStoppedAt = MacroSourceNode;
-
-									// Restore the debug object on the original blueprint if the flag was set
-									if (bResetObjectBeingDebuggedWhenFinished)
-									{
-										BlueprintObj->SetObjectBeingDebugged(SavedObjectBeingDebugged);
-									}
-
-									// Now switch to the macro source blueprint
-									BlueprintObj = MacroBlueprint;
-
-									// Ensure that the macro source blueprint's current debug object is set to the active object, and set the flag to restore it when finished
-									bResetObjectBeingDebuggedWhenFinished = true;
-									SavedObjectBeingDebugged = BlueprintObj->GetObjectBeingDebugged();
-									BlueprintObj->SetObjectBeingDebugged(const_cast<UObject*>(ActiveObject));
-								}
-							}
-						}
-					}
-				}
 			}
 
 			// Handle a breakpoint or single-step
@@ -494,20 +439,6 @@ UEdGraphNode* FKismetDebugUtilities::FindSourceNodeForCodeLocation(const UObject
 	return NULL;
 }
 
-UEdGraphNode* FKismetDebugUtilities::FindMacroSourceNodeForCodeLocation(const UObject* Object, UFunction* Function, int32 DebugOpcodeOffset)
-{
-	if (Object != NULL)
-	{
-		// Find the blueprint that corresponds to the object
-		if (UBlueprintGeneratedClass* Class = Cast<UBlueprintGeneratedClass>(FindClassForNode(Object, Function)))
-		{
-			return Class->GetDebugData().FindMacroSourceNodeFromCodeLocation(Function, DebugOpcodeOffset);
-		}
-	}
-
-	return NULL;
-}
-
 void FKismetDebugUtilities::CheckBreakConditions(UEdGraphNode* NodeStoppedAt, bool& InOutBreakExecution)
 {
 	FKismetDebugUtilitiesData& Data = FKismetDebugUtilitiesData::Get();
@@ -515,34 +446,33 @@ void FKismetDebugUtilities::CheckBreakConditions(UEdGraphNode* NodeStoppedAt, bo
 	if (NodeStoppedAt)
 	{
 		// Update tracked graph stack.
-		if (!Data.GraphStack.Num())
+		const UEdGraph* CurrentGraph = NodeStoppedAt->GetTypedOuter<UEdGraph>();
+		if (ensure(CurrentGraph))
 		{
-			Data.GraphStack.Push(NodeStoppedAt->GetTypedOuter<UEdGraph>());
-		}
-		else if (UK2Node_FunctionEntry* FunctionNode = Cast<UK2Node_FunctionEntry>(NodeStoppedAt))
-		{
-			UEdGraph* NewGraph = FunctionNode->GetTypedOuter<UEdGraph>();
-			const bool bAddFunction = Data.GraphStack.Num() && NewGraph ? NewGraph != Data.GraphStack.Last() : false;
-
-			if (bAddFunction)
+			if (Data.GraphStack.FindLast(CurrentGraph) != INDEX_NONE)
 			{
-				Data.GraphStack.Push(NewGraph);
+				while (Data.GraphStack.Last() != CurrentGraph)
+				{
+					Data.GraphStack.Pop();
+				}
 			}
-		}
-		else if (Data.GraphStack.Num())
-		{
-			UEdGraph* CurrGraph = NodeStoppedAt->GetTypedOuter<UEdGraph>();
-			if (Data.GraphStack.Last() != CurrGraph)
+			else
 			{
-				Data.GraphStack.Pop();
+				Data.GraphStack.Push(CurrentGraph);
 			}
 		}
+
 		// Figure out if we have a break condition.
 		if (!Data.bIsSingleStepping)
 		{
 			if (Data.TargetGraphStackDepth != INDEX_NONE)
 			{
 				InOutBreakExecution = Data.TargetGraphStackDepth >= Data.GraphStack.Num();
+				if (InOutBreakExecution)
+				{
+					// Reset once we've reached the target stack depth.
+					Data.TargetGraphStackDepth = INDEX_NONE;
+				}
 			}
 		}
 	}
@@ -605,7 +535,7 @@ void FKismetDebugUtilities::AttemptToBreakExecution(UBlueprint* BlueprintObj, co
 	// A check to !GIsAutomationTesting was removed from here as it seemed redundant.
 	// Breakpoints have to be explicitly enabled by the user which shouldn't happen 
 	// under automation and this was preventing debugging on automation test bp's.
-	if ((GUnrealEd->PlayWorld != NULL))
+	if ((GUnrealEd->PlayWorld != NULL) && NodeStoppedAt)
 	{
 		// Pause the simulation
 		GUnrealEd->PlayWorld->bDebugPauseExecution = true;
@@ -624,7 +554,7 @@ void FKismetDebugUtilities::AttemptToBreakExecution(UBlueprint* BlueprintObj, co
 	{
 		Data.LastExceptionMessage = Info.GetDescription();
 		FKismetEditorUtilities::BringKismetToFocusAttentionOnObject(NodeStoppedAt);
-
+		CallStackViewer::UpdateDisplayedCallstack(FBlueprintExceptionTracker::Get().ScriptStack);
 		FSlateApplication::Get().EnterDebuggingMode();
 	}
 }
diff --git a/Engine/Source/Editor/UnrealEd/Private/ReferenceInfoUtils.cpp b/Engine/Source/Editor/UnrealEd/Private/ReferenceInfoUtils.cpp
--- a/Engine/Source/Editor/UnrealEd/Private/ReferenceInfoUtils.cpp
+++ b/Engine/Source/Editor/UnrealEd/Private/ReferenceInfoUtils.cpp
@@ -476,7 +476,7 @@ namespace ReferenceInfoUtils
 		}
 
 		// Create a string for the resource size.
-		const SIZE_T ReferencedObjectResourceSize = ReferencedObject->GetResourceSizeBytes(EResourceSizeMode::Exclusive);
+		const SIZE_T ReferencedObjectResourceSize = ReferencedObject->GetResourceSizeBytes(EResourceSizeMode::EstimatedTotal);
 		FString ResourceSizeString;
 		if ( ReferencedObjectResourceSize > 0 )
 		{
diff --git a/Engine/Source/Editor/UnrealEd/Public/Kismet2/BlueprintEditorUtils.h b/Engine/Source/Editor/UnrealEd/Public/Kismet2/BlueprintEditorUtils.h
--- a/Engine/Source/Editor/UnrealEd/Public/Kismet2/BlueprintEditorUtils.h
+++ b/Engine/Source/Editor/UnrealEd/Public/Kismet2/BlueprintEditorUtils.h
@@ -22,6 +22,7 @@ class FCompilerResultsLog;
 class INameValidatorInterface;
 class UActorComponent;
 class UBlueprintGeneratedClass;
+class UK2Node_Event;
 class UK2Node_Variable;
 class ULevelScriptBlueprint;
 class USCS_Node;
@@ -380,6 +381,9 @@ public:
 	/** Adds an ubergraph page to this blueprint */
 	static void AddUbergraphPage(UBlueprint* Blueprint, class UEdGraph* Graph);
 
+	/** Returns the name of the Ubergraph Function that the provided blueprint uses */
+	static FName GetUbergraphFunctionName(const UBlueprint* ForBlueprint);
+
 	/** Adds a domain-specific graph to this blueprint */
 	static void AddDomainSpecificGraph(UBlueprint* Blueprint, class UEdGraph* Graph);
 
@@ -437,11 +441,10 @@ public:
 	static bool IsGraphReadOnly(UEdGraph* InGraph);
 
 	/** Look to see if an event already exists to override a particular function */
-	static class UK2Node_Event* FindOverrideForFunction(const UBlueprint* Blueprint, const UClass* SignatureClass, FName SignatureName);
+	static UK2Node_Event* FindOverrideForFunction(const UBlueprint* Blueprint, const UClass* SignatureClass, FName SignatureName);
 
 	/** Find the Custom Event if it already exists in the Blueprint */
-	static class UK2Node_Event* FindCustomEventNode(const UBlueprint* Blueprint, FName const CustomName);
-
+	static UK2Node_Event* FindCustomEventNode(const UBlueprint* Blueprint, FName const CustomName);
 
 	/** Returns all nodes in all graphs of the specified class */
 	template< class T > 
diff --git a/Engine/Source/Editor/UnrealEd/Public/Kismet2/CompilerResultsLog.h b/Engine/Source/Editor/UnrealEd/Public/Kismet2/CompilerResultsLog.h
--- a/Engine/Source/Editor/UnrealEd/Public/Kismet2/CompilerResultsLog.h
+++ b/Engine/Source/Editor/UnrealEd/Public/Kismet2/CompilerResultsLog.h
@@ -110,22 +110,13 @@ public:
 protected:
 	// Maps from transient object created during compiling to original 'source code' object
 	FBacktrackMap SourceBacktrackMap;
-	
-	// Maps immediately back to source uobject, which may itself by an intermediate object:
-	TMap<UEdGraphNode const*, UEdGraphNode*> FullSourceBacktrackMap;
 
 	// Name of the source object being compiled
 	FString SourcePath;
 
-	// Backtrack map for intermendiate tunnel/macro nodes
-	FBacktrackMap FinalNodeBackToTunnelSourceMap;
-	
 	// Map to track intermediate tunnel nodes back to the intermediate expansion tunnel instance.
 	TMap<TWeakObjectPtr<UEdGraphNode>, TWeakObjectPtr<UEdGraphNode>> IntermediateTunnelNodeToTunnelInstanceMap;
 
-	// Map to track active nested tunnels for intermediate tunnel instances.
-	TMultiMap<TWeakObjectPtr<UEdGraphNode>, TWeakObjectPtr<UEdGraphNode>> IntermediateTunnelInstanceHierarchyMap;
-
 public:
 	FCompilerResultsLog(bool bIsCompatibleWithEvents = true);
 	virtual ~FCompilerResultsLog();
@@ -260,31 +251,16 @@ public:
 	void NotifyIntermediateObjectCreation(UObject* NewObject, UObject* SourceObject);
 	void NotifyIntermediatePinCreation(UEdGraphPin* NewObject, UEdGraphPin* SourceObject);
 
-	/** Registers intermediate tunnel nodes, both the node and tunnel instance should be intermediate nodes */
-	void RegisterIntermediateTunnelNode(UEdGraphNode* Node, UEdGraphNode* OwningTunnelInstance);
-
-	/** Registers an intermediate tunnel instance node, the active tunnels should not be intermediate nodes. */
-	void RegisterIntermediateTunnelInstance(UEdGraphNode* IntermediateTunnel, TArray<TWeakObjectPtr<UEdGraphNode>>& ActiveTunnels);
-
-	/** Returns the source tunnel instance or the source tunnel node depending on what created the intermediate node */
-	UEdGraphNode* GetSourceNode(const UEdGraphNode* IntermediateNode);
-
-	/** Returns the intermediate tunnel instance that generated the node */
-	UEdGraphNode* GetIntermediateTunnelInstance(const UEdGraphNode* IntermediateNode) const;
-
-	/** Returns the source tunnel node for the intermediate node */
-	UEdGraphNode* GetSourceTunnelNode(const UEdGraphNode* IntermediateNode);
-
-	/** Returns the source tunnel instance that generated the intermediate node */
-	UEdGraphNode* GetSourceTunnelInstance(const UEdGraphNode* IntermediateNode);
-
-	/** Returns the active tunnel instances for the intermediate node */
-	void GetTunnelsActiveForNode(const UEdGraphNode* IntermediateNode, TArray<TWeakObjectPtr<UEdGraphNode>>& ActiveTunnelsOut);
+	/** Update the expansion map to note that Node was expanded from OuterTunnelInstance, both the node and tunnel instance should be intermediate nodes */
+	void NotifyIntermediateTunnelNode(UEdGraphNode* Node, UEdGraphNode* OuterTunnelInstance);
 
 	/** Returns the true source object for the passed in object */
 	UObject* FindSourceObject(UObject* PossiblyDuplicatedObject);
 	UObject const* FindSourceObject(UObject const* PossiblyDuplicatedObject) const;
 
+	/** Returns the intermediate tunnel instance that generated the node */
+	UEdGraphNode* GetIntermediateTunnelInstance(const UEdGraphNode* IntermediateNode) const;
+
 	/** Returns a int32 used to uniquely identify an action for the latent action manager */
 	int32 CalculateStableIdentifierForLatentActionManager( const UEdGraphNode* Node );
 
diff --git a/Engine/Source/Editor/UnrealEd/Public/Kismet2/KismetDebugUtilities.h b/Engine/Source/Editor/UnrealEd/Public/Kismet2/KismetDebugUtilities.h
--- a/Engine/Source/Editor/UnrealEd/Public/Kismet2/KismetDebugUtilities.h
+++ b/Engine/Source/Editor/UnrealEd/Public/Kismet2/KismetDebugUtilities.h
@@ -131,9 +131,6 @@ public:
 	// Find the node that resulted in code at the specified location in the Object, or NULL if there was a problem (e.g., no debugging information was generated)
 	static class UEdGraphNode* FindSourceNodeForCodeLocation(const UObject* Object, UFunction* Function, int32 DebugOpcodeOffset, bool bAllowImpreciseHit = false);
 
-	// Find the macro node that resulted in code at the specified location in the Object, or NULL if there was a problem (e.g., no debugging information was generated)
-	static class UEdGraphNode* FindMacroSourceNodeForCodeLocation(const UObject* Object, UFunction* Function, int32 DebugOpcodeOffset);
-
 	// Return proper class for breakpoint
 	static UClass* FindClassForNode(const UObject* Object, UFunction* Function);
 
diff --git a/Engine/Source/Editor/UnrealEd/Public/Settings/EditorProjectSettings.h b/Engine/Source/Editor/UnrealEd/Public/Settings/EditorProjectSettings.h
--- a/Engine/Source/Editor/UnrealEd/Public/Settings/EditorProjectSettings.h
+++ b/Engine/Source/Editor/UnrealEd/Public/Settings/EditorProjectSettings.h
@@ -172,6 +172,10 @@ public:
 	UPROPERTY(EditAnywhere, config, Category=Blueprints)
 	uint32 bSkipUnneededDependencyCompilation:1;
 
+	/** If enabled, the editor will load packages to look for soft references to actors when deleting/renaming them. This can be slow in large projects so disable this to improve performance but increase the chance of breaking blueprints/sequences that use soft actor references */
+	UPROPERTY(EditAnywhere, config, Category=Actors)
+	uint32 bValidateUnloadedSoftActorReferences : 1;
+
 	// UObject interface
 	virtual void PostEditChangeProperty(struct FPropertyChangedEvent& PropertyChangedEvent) override;
 	// End of UObject interface
diff --git a/Engine/Source/Programs/AutomationTool/Scripts/CopyBuildToStagingDirectory.Automation.cs b/Engine/Source/Programs/AutomationTool/Scripts/CopyBuildToStagingDirectory.Automation.cs
--- a/Engine/Source/Programs/AutomationTool/Scripts/CopyBuildToStagingDirectory.Automation.cs
+++ b/Engine/Source/Programs/AutomationTool/Scripts/CopyBuildToStagingDirectory.Automation.cs
@@ -309,11 +309,9 @@ public partial class Project : CommandUtils
 
 			// Put all of the cooked dir into the staged dir
 			DirectoryReference PlatformCookDir = String.IsNullOrEmpty(Params.CookOutputDir) ? DirectoryReference.Combine(DLCRoot, "Saved", "Cooked", SC.CookPlatform) : DirectoryReference.Combine(new DirectoryReference(Params.CookOutputDir), SC.CookPlatform);
-
-			// Stage any loose files in the root folder
-			// TODO: not sure if we should stage the loose files if we have pak files enabled... 
-			SC.StageFiles(StagedFileType.UFS, PlatformCookDir, StageFilesSearch.TopDirectoryOnly, SC.RelativeProjectRootForStage);
-
+			DirectoryReference PlatformEngineDir = DirectoryReference.Combine(PlatformCookDir, "Engine");
+			DirectoryReference PlatformMetadataDir = DirectoryReference.Combine(PlatformCookDir, SC.ShortProjectName, "Metadata");
+			
 			SC.StageFiles(StagedFileType.UFS, DirectoryReference.Combine(DLCRoot, "Config"), "*.ini", StageFilesSearch.AllDirectories);
 
 			if (Params.DLCActLikePatch)
@@ -325,19 +323,28 @@ public partial class Project : CommandUtils
 				SC.RemapDirectories.Add(Tuple.Create(new StagedDirectoryReference(RelativeDLCContentPath), new StagedDirectoryReference(RelativeRootContentPath)));
 			}
 
-			StagedDirectoryReference EngineDirectory = new StagedDirectoryReference("Engine");
-			// Stage each sub directory separately so that we can skip Engine if need be
-			foreach (DirectoryReference SubDir in DirectoryReference.EnumerateDirectories(PlatformCookDir))
+			// Stage all the cooked data, this is the same rule as normal stage except we may skip Engine
+			List<FileReference> CookedFiles = DirectoryReference.EnumerateFiles(PlatformCookDir, "*", SearchOption.AllDirectories).ToList();
+			foreach (FileReference CookedFile in CookedFiles)
 			{
-				StagedDirectoryReference MountPoint = new StagedDirectoryReference(SubDir.MakeRelativeTo(PlatformCookDir));
-				// Dedicated server cook doesn't save shaders so no Engine dir is created
+				// Skip metadata directory
+				if (CookedFile.Directory.IsUnderDirectory(PlatformMetadataDir))
+				{
+					continue;
+				}
 
-				if ((!SC.DedicatedServer) && (!Params.DLCIncludeEngineContent) && MountPoint == EngineDirectory)
+				// Dedicated server cook doesn't save shaders so no Engine dir is created
+				if ((!SC.DedicatedServer) && (!Params.DLCIncludeEngineContent) && CookedFile.Directory.IsUnderDirectory(PlatformEngineDir))
 				{
 					continue;
 				}
 
-				SC.StageFiles(StagedFileType.UFS, SubDir, "*", StageFilesSearch.AllDirectories, MountPoint);
+				// json files have never been staged
+				// metallib files cannot *currently* be staged as UFS as the Metal API needs to mmap them from files on disk in order to function efficiently
+				if (!CookedFile.HasExtension(".json") && !CookedFile.HasExtension(".metallib"))
+				{
+					SC.StageFile(StagedFileType.UFS, CookedFile, new StagedFileReference(CookedFile.MakeRelativeTo(PlatformCookDir)));
+				}
 			}
 
 			FileReference PluginSettingsFile = FileReference.Combine(DLCRoot, "Config", "PluginSettings.ini");
@@ -547,11 +554,11 @@ public partial class Project : CommandUtils
 					}
 				}
 
-			// shader cache
+				// shader cache
 				{
-				DirectoryReference ShaderCacheRoot = DirectoryReference.Combine(SC.ProjectRoot, "Content");
-				List<FileReference> ShaderCacheFiles = SC.FindFilesToStage(ShaderCacheRoot, "DrawCache-*.ushadercache", StageFilesSearch.TopDirectoryOnly);
-				SC.StageFiles(StagedFileType.UFS, ShaderCacheFiles);
+					DirectoryReference ShaderCacheRoot = DirectoryReference.Combine(SC.ProjectRoot, "Content");
+					List<FileReference> ShaderCacheFiles = SC.FindFilesToStage(ShaderCacheRoot, "DrawCache-*.ushadercache", StageFilesSearch.TopDirectoryOnly);
+					SC.StageFiles(StagedFileType.UFS, ShaderCacheFiles);
 				}
 
 				// Get the final output directory for cooked data
@@ -569,13 +576,21 @@ public partial class Project : CommandUtils
 					CookOutputDir = DirectoryReference.Combine(SC.ProjectRoot, "Saved", "Cooked", SC.CookPlatform);
 				}
 
+				DirectoryReference MetadataOutputDir = DirectoryReference.Combine(CookOutputDir, SC.ShortProjectName, "Metadata");
+
 				// Stage all the cooked data. Currently not filtering this by restricted folders, since we shouldn't mask invalid references by filtering them out.
 				List<FileReference> CookedFiles = DirectoryReference.EnumerateFiles(CookOutputDir, "*", SearchOption.AllDirectories).ToList();
 				foreach (FileReference CookedFile in CookedFiles)
 				{
-				// json files have never been staged
-				// metallib files cannot *currently* be staged as UFS as the Metal API needs to mmap them from files on disk in order to function efficiently
-				if(!CookedFile.HasExtension(".json") && !CookedFile.HasExtension(".metallib"))
+					// Skip metadata directory
+					if (CookedFile.Directory.IsUnderDirectory(MetadataOutputDir))
+					{
+						continue;
+					}
+
+					// json files have never been staged
+					// metallib files cannot *currently* be staged as UFS as the Metal API needs to mmap them from files on disk in order to function efficiently
+					if(!CookedFile.HasExtension(".json") && !CookedFile.HasExtension(".metallib"))
 					{
 						SC.StageFile(StagedFileType.UFS, CookedFile, new StagedFileReference(CookedFile.MakeRelativeTo(CookOutputDir)));
 					}
diff --git a/Engine/Source/Programs/AutomationTool/Scripts/SharedCookedBuild.cs b/Engine/Source/Programs/AutomationTool/Scripts/SharedCookedBuild.cs
--- a/Engine/Source/Programs/AutomationTool/Scripts/SharedCookedBuild.cs
+++ b/Engine/Source/Programs/AutomationTool/Scripts/SharedCookedBuild.cs
@@ -1,4 +1,4 @@
-﻿// Copyright 1998-2017 Epic Games, Inc. All Rights Reserved.
+// Copyright 1998-2017 Epic Games, Inc. All Rights Reserved.
 using System;
 using System.Collections.Generic;
 using System.IO;
@@ -85,16 +85,16 @@ public class SharedCookedBuild
 
 
 
-		string CookedBuildCookedDirectory = Path.Combine(CookedBuildPath, "Cooked");
-		CookedBuildCookedDirectory = Path.GetFullPath(CookedBuildCookedDirectory);
-		string LocalBuildCookedDirectory = Path.Combine(LocalPath, "Cooked");
-		LocalBuildCookedDirectory = Path.GetFullPath(LocalBuildCookedDirectory);
-		if (Directory.Exists(CookedBuildCookedDirectory))
+		string CookedBuildMetadataDirectory = Path.Combine(CookedBuildPath, "Metadata");
+		CookedBuildMetadataDirectory = Path.GetFullPath(CookedBuildMetadataDirectory);
+		string LocalBuildMetadataDirectory = Path.Combine(LocalPath, "Metadata");
+		LocalBuildMetadataDirectory = Path.GetFullPath(LocalBuildMetadataDirectory);
+		if (Directory.Exists(CookedBuildMetadataDirectory))
 		{
-			foreach (string FileName in Directory.EnumerateFiles(CookedBuildCookedDirectory, "*.*", SearchOption.AllDirectories))
+			foreach (string FileName in Directory.EnumerateFiles(CookedBuildMetadataDirectory, "*.*", SearchOption.AllDirectories))
 			{
 				string SourceFileName = Path.GetFullPath(FileName);
-				string DestFileName = SourceFileName.Replace(CookedBuildCookedDirectory, LocalBuildCookedDirectory);
+				string DestFileName = SourceFileName.Replace(CookedBuildMetadataDirectory, LocalBuildMetadataDirectory);
 				Directory.CreateDirectory(Path.GetDirectoryName(DestFileName));
 				File.Copy(SourceFileName, DestFileName);
 			}
diff --git a/Engine/Source/Programs/UnrealHeaderTool/Private/HeaderParser.cpp b/Engine/Source/Programs/UnrealHeaderTool/Private/HeaderParser.cpp
--- a/Engine/Source/Programs/UnrealHeaderTool/Private/HeaderParser.cpp
+++ b/Engine/Source/Programs/UnrealHeaderTool/Private/HeaderParser.cpp
@@ -3527,11 +3527,11 @@ void FHeaderParser::GetVarType(
 		{
 			if (0 != (CPF_DisableEditOnInstance & Flags))
 			{
-				UE_LOG_WARNING_UHT(TEXT("Property cannot have 'DisableEditOnInstance' or 'BlueprintReadOnly' and 'ExposeOnSpawn' flags"));
+				UE_LOG_WARNING_UHT(TEXT("Property cannot have both 'DisableEditOnInstance' and 'ExposeOnSpawn' flags"));
 			}
 			if (0 == (CPF_BlueprintVisible & Flags))
 			{
-				UE_LOG_WARNING_UHT(TEXT("Property cannot have 'ExposeOnSpawn' with 'BlueprintVisible' flag."));
+				UE_LOG_WARNING_UHT(TEXT("Property cannot have 'ExposeOnSpawn' without 'BlueprintVisible' flag."));
 			}
 			Flags |= CPF_ExposeOnSpawn;
 		}
diff --git a/Engine/Source/Runtime/AssetRegistry/Private/AssetRegistry.cpp b/Engine/Source/Runtime/AssetRegistry/Private/AssetRegistry.cpp
--- a/Engine/Source/Runtime/AssetRegistry/Private/AssetRegistry.cpp
+++ b/Engine/Source/Runtime/AssetRegistry/Private/AssetRegistry.cpp
@@ -86,23 +86,11 @@ UAssetRegistryImpl::UAssetRegistryImpl(const FObjectInitializer& ObjectInitializ
 	// for platforms that require cooked data, we attempt to load a premade asset registry
 	else if (FPlatformProperties::RequiresCookedData())
 	{
-		bool bLoadedDevelopment = false;
 		// load the cooked data
 		FArrayReader SerializedAssetData;
 
-#if !(UE_BUILD_SHIPPING || UE_BUILD_TEST)
-		// Allow loading development asset registry
-		FString DevAssetRegistryFilename = (FPaths::ProjectDir() / TEXT("DevelopmentAssetRegistry.bin"));
-		if (FParse::Param(FCommandLine::Get(), TEXT("LoadDevAssetRegistry")) && IFileManager::Get().FileExists(*DevAssetRegistryFilename) && FFileHelper::LoadFileToArray(SerializedAssetData, *DevAssetRegistryFilename))
-		{
-			SerializationOptions.ModifyForDevelopment();
-
-			Serialize(SerializedAssetData);
-			bLoadedDevelopment = true;
-		}
-#endif
 		FString AssetRegistryFilename = (FPaths::ProjectDir() / TEXT("AssetRegistry.bin"));
-		if (SerializationOptions.bSerializeAssetRegistry && !bLoadedDevelopment && IFileManager::Get().FileExists(*AssetRegistryFilename) && FFileHelper::LoadFileToArray(SerializedAssetData, *AssetRegistryFilename))
+		if (SerializationOptions.bSerializeAssetRegistry && IFileManager::Get().FileExists(*AssetRegistryFilename) && FFileHelper::LoadFileToArray(SerializedAssetData, *AssetRegistryFilename))
 		{
 			// serialize the data with the memory reader (will convert FStrings to FNames, etc)
 			Serialize(SerializedAssetData);
diff --git a/Engine/Source/Runtime/AssetRegistry/Private/AssetRegistryState.cpp b/Engine/Source/Runtime/AssetRegistry/Private/AssetRegistryState.cpp
--- a/Engine/Source/Runtime/AssetRegistry/Private/AssetRegistryState.cpp
+++ b/Engine/Source/Runtime/AssetRegistry/Private/AssetRegistryState.cpp
@@ -210,7 +210,7 @@ void FAssetRegistryState::InitializeFromExisting(const TMap<FName, FAssetData*>&
 					FAssetPackageData* NewData = CreateOrGetAssetPackageData(Pair.Key);
 					*NewData = *Pair.Value;
 				}
-				else if (CachedAssetsByPackageName.Find(Pair.Key) || bIsScriptPackage)
+				else if (CachedAssetsByPackageName.Find(Pair.Key))
 				{
 					FAssetPackageData* NewData = CreateOrGetAssetPackageData(Pair.Key);
 					*NewData = *Pair.Value;
diff --git a/Engine/Source/Runtime/AssetRegistry/Public/AssetData.h b/Engine/Source/Runtime/AssetRegistry/Public/AssetData.h
--- a/Engine/Source/Runtime/AssetRegistry/Public/AssetData.h
+++ b/Engine/Source/Runtime/AssetRegistry/Public/AssetData.h
@@ -444,6 +444,11 @@ private:
 	}
 };
 
+FORCEINLINE uint32 GetTypeHash(const FAssetData& AssetData)
+{
+	return GetTypeHash(AssetData.ObjectPath);
+}
+
 
 template<>
 struct TStructOpsTypeTraits<FAssetData> : public TStructOpsTypeTraitsBase2<FAssetData>
