Copying //UE4/Dev-Framework to //UE4/Dev-Main (Source: //UE4/Dev-Framework @ 3803685)
#rnx

============================
  MAJOR FEATURES & CHANGES
============================

Change 3794434 by [REDACTED_PERSON]

	Correctly call PostPlacedNewNode before AllocateDefaultPins in 2 places that had the order wrong

Change 3796912 by [REDACTED_PERSON]

	Added several transform related functions to the HIDE_ACTOR_TRANSFORM_FUNCTIONS macro and fixed the signatures on some that were already there.

Change 3781456 by [REDACTED_PERSON]

	Disable net flags related assert when linking reflection data that is being loaded

Change 3781725 by mason.seay

	Moved test blueprints to subfolder

Change 3781828 by mason.seay

	Removed some unnecessary blueprint logic

Change 3783304 by [REDACTED_PERSON]

	More descriptive log output after an attempt to break execution on a missing function/interface.

	#4238

Change 3783678 by [REDACTED_PERSON]

	Fix for a serialization oversight in UEdGraphPin that could lead to a compile error during Blueprint nativization.

Change 3783708 by [REDACTED_PERSON]

Change 3783922 by [REDACTED_PERSON]

	If a game was depending on this functionality to keep things loaded, change your blueprint to set the return value into a variable, that will keep it hard referenced normally

Change 3784604 by [REDACTED_PERSON]

Change 3785039 by [REDACTED_PERSON]

	Fix it so empty pins do not display "None" as their label

Change 3785794 by [REDACTED_PERSON]

	Guard against a reported crash when setting bIsEditorOnly flag on attached scene components.

	Change summary:
	- Modified SetDescendantIsEditorOnly() to add a NULL check around an iteration (consistent with other places that iterate the attachment list).

Change 3785842 by [REDACTED_PERSON]

	Guard against a reported crash when deleting one or more Actor components.

	Change summary:
	- Modified FComponentEditorUtils::DeleteComponents() to add a NULL check when iterating through child attachments.

Change 3785851 by [REDACTED_PERSON]

	Guard against a reported crash when deleting a non-root scene component with no parent attachment.

	Change summary:
	- Modified USceneComponent::DestroyComponent() to use an ensure() instead of a check() to allow the deletion to proceed but still provide a warning to potentially help catch this later.

Change 3785852 by [REDACTED_PERSON]

	CIS fix.

Change 3786141 by [REDACTED_PERSON]

	Removed an invalid comment emitted to codegen during nativization of a BlueprintNativeEvent override signature.

	Was intended to be submitted along with CL# 3783678.

Change 3788496 by [REDACTED_PERSON]

	User structs now work identically to native structs when it comes to default values, it is set at struct initialization time instead of at tagged property serialization time
	Remove many editor-time UDS hacks now that they work like native. The only special code left is that it forces delta serialization on in the editor, to deal with the defaults changing

Change 3789308 by [REDACTED_PERSON]

	Don't unregister/register components and ticks for an Actor just testing a rename

Change 3789644 by [REDACTED_PERSON]

	Static analysis fix

Change 3789777 by [REDACTED_PERSON]

Change 3789908 by mason.seay

	Test struct using containers

Change 3790475 by [REDACTED_PERSON]

	Fix user structs to work properly with circular reference checking. When we copy the values out of the default we may need to register the newly created structs with the circular reference system.
	I'm using raw pointers here as after investigation the fixup happens before any relevant freeing could occur

Change 3790777 by [REDACTED_PERSON]

	Partially address UE-53086, it now correctly updates map properties when they are removed, but readding them can add extra entries in some case because it loses the "deleted" information

Change 3790919 by [REDACTED_PERSON]

	Add start of User Structs functional test

Change 3792050 by [REDACTED_PERSON]

	Added several transform related functions to the HIDE_ACTOR_TRANSFORM_FUNCTIONS macro and fixed the signatures on some that were already there.

Change 3793542 by [REDACTED_PERSON]

	Back out changelist 3792050

Change 3794193 by [REDACTED_PERSON]

	Never reset loaders when renaming the child actor to avoid cases where IsLoading returns false but flushing loading will still cause issues due to async loading

Change 3794211 by [REDACTED_PERSON]

	Don't include redirectors or otherwise invalid .umap files as sub levels for world composition

Change 3794564 by [REDACTED_PERSON]

	Make it safe to call RecombinePin on the top level split pin

Change 3794851 by [REDACTED_PERSON]

	Fix split pins on macro instances for resolved wildcard pins not resetting when disconnected

Change 3795017 by [REDACTED_PERSON]

	Properly enforce all actor subclasses not ticking when tick type is viewports only unless ShouldTickIfViewportsOnly returns true

Change 3795143 by [REDACTED_PERSON]

	After copying properties clear the actor label so that the child actor template's display name is not misleading

Change 3795241 by [REDACTED_PERSON]

	Fix subsequent navigation operations after navigating to a window that is already open

Change 3795255 by [REDACTED_PERSON]

	Harden pin serialization logic, specifically to prevent crash due to bad 'conform' routines when renaming functions

Change 3795660 by [REDACTED_PERSON]

	Blueprints that have no parent class can now be reparented when using the compilation manager

Change 3796496 by [REDACTED_PERSON]

	Fix without editor compile errors

Change 3797316 by [REDACTED_PERSON]

	Downgrade ensure to warning as we have known ways of triggering

Change 3797619 by [REDACTED_PERSON]

	Reinstate a faster lookup mechanism for tunnel instance source nodes during marching ants draw calls when stopped at a breakpoint, and fix issues with missing marching ants when stopped at a breakpoint in a tunnel instance graph.

	Change summary:
	- Removed FindMatchingTunnelInstanceNode() utility method in BlueprintConnectionDrawingPolicy.cpp (no longer needed).
	- Restored/renamed FDebuggingInfoForSingleFunction::LineNumberToTunnellInstanceSourceNodeMap.
	- Modified FBlueprintDebugData::RegisterNodeToCodeAssociation() to restore the parameter for the tunnel instance node chain and map it to the current code offset if not empty.
	- Modified FScriptBytecodeWriter::EmitInstrumentation() to construct the tunnel instance node chain that will be mapped to the current code offset if not empty.
	- Modified FKismetConnectionDrawingPolicy::BuildExecutionRoadmap() to more efficiently find tunnel instance nodes through the table that now maps tunnel instance node chains to the current script code offset.
	- Modified the tunnel node case in FKismetCompilerContext::ExpandTunnelsAndMacros() to better handle nested tunnel instance expansions (which may expand to include internal tunnel nodes in different orders).
	- Modified FKismetCompilerContext::MapExpansionPathToTunnelInstance() to avoid incorrectly remapping nested composite tunnel expansions to the outermost intermediate tunnel instance node.
	- Modified FKismetCompilerContext::ProcessIntermediateTunnelBoundary() to map intermediate tunnel boundary nodes back to the outer intermediate tunnel instance node from which they were expanded.
	- Modified FKismetCompilerContext::ProcessIntermediateTunnelBoundary() to better position the boundary nodes in the intermediate graph.
	- Added a few verbose comment blocks in order to hopefully help visualize how tunnel expansion debugging is currently implemented and also to remind myself later.

Change 3801237 by [REDACTED_PERSON]

Change 3801768 by [REDACTED_PERSON]

	Add more tests to UserStructs functional test

Change 3801794 by [REDACTED_PERSON]

	Mirror 3794923/3794318, make sure replication data is initialized correctly across editor/EDL/compilation manager/no compilation manager load flows

Change 3803091 by [REDACTED_PERSON]

	Fix issue where nativized user structs had the wrong defaults, by correctly diffing the struct defaults and placing them in the constructor instead of a static GetDefaultValue function
	I had to override struct constructors to not call LoadObject for object properties, FindObject works because the dependencies have already been queued for loading by the referencing blueprint in EDL

Change 3803659 by [REDACTED_PERSON]

	Fix static analysis warning

Change 3803685 by [REDACTED_PERSON]

	Fix extra )

diff --git a/Engine/Source/Developer/BlueprintCompilerCppBackend/Private/BlueprintCompilerCppBackendBase.cpp b/Engine/Source/Developer/BlueprintCompilerCppBackend/Private/BlueprintCompilerCppBackendBase.cpp
--- a/Engine/Source/Developer/BlueprintCompilerCppBackend/Private/BlueprintCompilerCppBackendBase.cpp
+++ b/Engine/Source/Developer/BlueprintCompilerCppBackend/Private/BlueprintCompilerCppBackendBase.cpp
@@ -835,7 +835,8 @@ FString FBlueprintCompilerCppBackendBase::GenerateArgList(const FEmitterLocalCon
 		}
 		else
 		{
-			if (ArgProperty->HasAnyPropertyFlags(CPF_OutParm))
+			if (ArgProperty->HasAnyPropertyFlags(CPF_OutParm)
+				&& !ArgProperty->HasAnyPropertyFlags(CPF_ConstParm))
 			{
 				ArgListStr += TEXT("/*out*/ ");
 			}
@@ -1014,7 +1015,7 @@ void FBlueprintCompilerCppBackendBase::GenerateCodeFromStruct(UUserDefinedStruct
 		EmitterContext.Header.AddLine(TEXT("GENERATED_BODY()"));
 		EmitStructProperties(EmitterContext, SourceStruct);
 
-		FEmitDefaultValueHelper::GenerateGetDefaultValue(SourceStruct, EmitterContext);
+		FEmitDefaultValueHelper::GenerateUserStructConstructor(SourceStruct, EmitterContext);
 
 		EmitterContext.Header.AddLine(TEXT("static void __StaticDependenciesAssets(TArray<FBlueprintDependencyData>& AssetsToLoad);"));
 		EmitterContext.Header.AddLine(TEXT("static void __StaticDependencies_DirectlyUsedAssets(TArray<FBlueprintDependencyData>& AssetsToLoad);"));
diff --git a/Engine/Source/Developer/BlueprintCompilerCppBackend/Private/BlueprintCompilerCppBackendUtils.cpp b/Engine/Source/Developer/BlueprintCompilerCppBackend/Private/BlueprintCompilerCppBackendUtils.cpp
--- a/Engine/Source/Developer/BlueprintCompilerCppBackend/Private/BlueprintCompilerCppBackendUtils.cpp
+++ b/Engine/Source/Developer/BlueprintCompilerCppBackend/Private/BlueprintCompilerCppBackendUtils.cpp
@@ -1,4 +1,4 @@
-﻿// Copyright 1998-2017 Epic Games, Inc. All Rights Reserved.
+// Copyright 1998-2017 Epic Games, Inc. All Rights Reserved.
 
 #include "BlueprintCompilerCppBackendUtils.h"
 #include "Misc/App.h"
@@ -40,6 +40,7 @@ FString FEmitterLocalContext::FindGloballyMappedObject(const UObject* Object, co
 		}
 	}
 
+	UUserDefinedStruct* ActualUserStruct = Cast<UUserDefinedStruct>(Dependencies.GetActualStruct());
 	UClass* ActualClass = Cast<UClass>(Dependencies.GetActualStruct());
 	UClass* OriginalActualClass = Dependencies.FindOriginalClass(ActualClass);
 	UClass* OuterClass = Object ? Cast<UClass>(Object->GetOuter()) : nullptr;	// SCS component templates will have an Outer that equates to their owning BPGC; since they're not currently DSOs, we have to special-case them.
@@ -259,6 +260,15 @@ FString FEmitterLocalContext::FindGloballyMappedObject(const UObject* Object, co
 		}
 	}
 
+	if (ActualUserStruct)
+	{
+		// For user structs, the default action of loading is unsafe so call the wrapper function
+		return FString::Printf(TEXT("CastChecked<%s>(FConvertedBlueprintsDependencies::LoadObjectForStructConstructor(%s::StaticStruct(),TEXT(\"%s\")), ECastCheckedType::NullAllowed)")
+			, *ClassString()
+			, *FEmitHelper::GetCppName(ActualUserStruct)
+			, *GetPathNameSafe(Object).ReplaceCharWithEscapedChar());			
+	}
+
 	return FString{};
 }
 
@@ -1104,7 +1114,7 @@ FString FEmitHelper::LiteralTerm(FEmitterLocalContext& EmitterContext, const FEd
 				EmitterContext.AddLine(FString::Printf(TEXT("auto %s = %s%s;")
 					, *LocalStructNativeName
 					, *StructName
-					, AsUDS ? TEXT("::GetDefaultValue()") : FEmitHelper::EmptyDefaultConstructor(StructType)));
+					, FEmitHelper::EmptyDefaultConstructor(StructType)));
 				if (AsUDS)
 				{
 					EmitterContext.StructsWithDefaultValuesUsed.Add(AsUDS);
diff --git a/Engine/Source/Developer/BlueprintCompilerCppBackend/Private/BlueprintCompilerCppBackendUtils.h b/Engine/Source/Developer/BlueprintCompilerCppBackend/Private/BlueprintCompilerCppBackendUtils.h
--- a/Engine/Source/Developer/BlueprintCompilerCppBackend/Private/BlueprintCompilerCppBackendUtils.h
+++ b/Engine/Source/Developer/BlueprintCompilerCppBackend/Private/BlueprintCompilerCppBackendUtils.h
@@ -1,4 +1,4 @@
-﻿// Copyright 1998-2017 Epic Games, Inc. All Rights Reserved.
+// Copyright 1998-2017 Epic Games, Inc. All Rights Reserved.
 #pragma  once
 
 #include "CoreMinimal.h"
@@ -321,7 +321,7 @@ struct FNonativeComponentData;
 
 struct FEmitDefaultValueHelper
 {
-	static void GenerateGetDefaultValue(const UUserDefinedStruct* Struct, FEmitterLocalContext& EmitterContext);
+	static void GenerateUserStructConstructor(const UUserDefinedStruct* Struct, FEmitterLocalContext& EmitterContext);
 
 	static void GenerateConstructor(FEmitterLocalContext& Context);
 
diff --git a/Engine/Source/Developer/BlueprintCompilerCppBackend/Private/BlueprintCompilerCppBackendValueHelper.cpp b/Engine/Source/Developer/BlueprintCompilerCppBackend/Private/BlueprintCompilerCppBackendValueHelper.cpp
--- a/Engine/Source/Developer/BlueprintCompilerCppBackend/Private/BlueprintCompilerCppBackendValueHelper.cpp
+++ b/Engine/Source/Developer/BlueprintCompilerCppBackend/Private/BlueprintCompilerCppBackendValueHelper.cpp
@@ -1,4 +1,4 @@
-﻿// Copyright 1998-2017 Epic Games, Inc. All Rights Reserved.
+// Copyright 1998-2017 Epic Games, Inc. All Rights Reserved.
 
 #include "CoreMinimal.h"
 #include "Misc/Guid.h"
@@ -133,32 +133,28 @@ void FEmitDefaultValueHelper::OuterGenerate(FEmitterLocalContext& Context
 	}
 }
 
-void FEmitDefaultValueHelper::GenerateGetDefaultValue(const UUserDefinedStruct* Struct, FEmitterLocalContext& Context)
+void FEmitDefaultValueHelper::GenerateUserStructConstructor(const UUserDefinedStruct* Struct, FEmitterLocalContext& Context)
 {
 	check(Struct);
 	const FString StructName = FEmitHelper::GetCppName(Struct);
 
 	// Declaration
-	Context.Header.AddLine(FString::Printf(TEXT("static %s GetDefaultValue();"), *StructName));
+	Context.Header.AddLine(FString::Printf(TEXT("%s();"), *StructName));
 
 	// Definition
-	Context.Body.AddLine(FString::Printf(TEXT("%s %s::GetDefaultValue()"), *StructName, *StructName));
+	Context.Body.AddLine(FString::Printf(TEXT("%s::%s()"), *StructName, *StructName));
 	Context.Body.AddLine(TEXT("{"));
 
 	Context.Body.IncreaseIndent();
-	Context.Body.AddLine(FString::Printf(TEXT("FStructOnScope StructOnScope(%s::StaticStruct());"), *StructName));
-	Context.Body.AddLine(FString::Printf(TEXT("%s& DefaultData__ = *((%s*)StructOnScope.GetStructMemory());"), *StructName, *StructName));
 	{
 		TGuardValue<FCodeText*> OriginalDefaultTarget(Context.DefaultTarget, &Context.Body);
 		FStructOnScope StructData(Struct);
-		FStructureEditorUtils::Fill_MakeStructureDefaultValue(Struct, StructData.GetStructMemory());
-		FStructOnScope RawDefaultStructOnScope(Struct);
+		FUserStructOnScopeIgnoreDefaults RawDefaultStructOnScope(Struct);
 		for (auto Property : TFieldRange<const UProperty>(Struct))
 		{
-			OuterGenerate(Context, Property, TEXT("DefaultData__"), StructData.GetStructMemory(), RawDefaultStructOnScope.GetStructMemory(), EPropertyAccessOperator::Dot);
+			OuterGenerate(Context, Property, TEXT(""), StructData.GetStructMemory(), RawDefaultStructOnScope.GetStructMemory(), EPropertyAccessOperator::None);
 		}
 	}
-	Context.Body.AddLine(TEXT("return DefaultData__;"));
 	Context.Body.DecreaseIndent();
 
 	Context.Body.AddLine(TEXT("}"));
diff --git a/Engine/Source/Editor/BlueprintGraph/Classes/K2Node.h b/Engine/Source/Editor/BlueprintGraph/Classes/K2Node.h
--- a/Engine/Source/Editor/BlueprintGraph/Classes/K2Node.h
+++ b/Engine/Source/Editor/BlueprintGraph/Classes/K2Node.h
@@ -191,6 +191,7 @@ class UK2Node : public UEdGraphNode
 	BLUEPRINTGRAPH_API virtual UEdGraphPin* GetPassThroughPin(const UEdGraphPin* FromPin) const override;
 	BLUEPRINTGRAPH_API virtual bool IsInDevelopmentMode() const override;
 	BLUEPRINTGRAPH_API virtual void ValidateNodeDuringCompilation(class FCompilerResultsLog& MessageLog) const override;
+	BLUEPRINTGRAPH_API virtual FString GetPinMetaData(FName InPinName, FName InKey) override;
 	// End of UEdGraphNode interface
 
 	// K2Node interface
diff --git a/Engine/Source/Editor/BlueprintGraph/Classes/K2Node_Tunnel.h b/Engine/Source/Editor/BlueprintGraph/Classes/K2Node_Tunnel.h
--- a/Engine/Source/Editor/BlueprintGraph/Classes/K2Node_Tunnel.h
+++ b/Engine/Source/Editor/BlueprintGraph/Classes/K2Node_Tunnel.h
@@ -5,6 +5,7 @@
 #include "CoreMinimal.h"
 #include "UObject/ObjectMacros.h"
 #include "K2Node_EditablePinBase.h"
+#include "EdGraph/EdGraphPin.h"
 #include "K2Node_Tunnel.generated.h"
 
 UCLASS()
@@ -17,11 +18,11 @@ class BLUEPRINTGRAPH_API UK2Node_Tunnel : public UK2Node_EditablePinBase
 	
 	// The output pins of this tunnel node came from the input pins of OutputSourceNode
 	UPROPERTY()
-	class UK2Node_Tunnel* OutputSourceNode;
+	UK2Node_Tunnel* OutputSourceNode;
 
 	// The input pins of this tunnel go to the output pins of InputSinkNode
 	UPROPERTY()
-	class UK2Node_Tunnel* InputSinkNode;
+	UK2Node_Tunnel* InputSinkNode;
 
 	// Whether this node is allowed to have inputs
 	UPROPERTY()
@@ -71,6 +72,13 @@ protected:
 	* @param bInAllWildcardPinsUnlinked	TRUE if all wildcard pins were unlinked
 	*/
 	virtual void PostFixupAllWildcardPins(bool bInAllWildcardPinsUnlinked) {}
+
+	// Utility function that subclasses must call after allocating their default pins
+	void CacheWildcardPins();
+
+	// Cache of the pins that were created as wildcards
+	TArray<UEdGraphPin*> WildcardPins;
+
 public:
 	// The input pins of this tunnel go to the output pins of InputSinkNode (can be NULL).
 	virtual UK2Node_Tunnel* GetInputSink() const;
diff --git a/Engine/Source/Editor/BlueprintGraph/Private/EdGraphSchema_K2.cpp b/Engine/Source/Editor/BlueprintGraph/Private/EdGraphSchema_K2.cpp
--- a/Engine/Source/Editor/BlueprintGraph/Private/EdGraphSchema_K2.cpp
+++ b/Engine/Source/Editor/BlueprintGraph/Private/EdGraphSchema_K2.cpp
@@ -6795,9 +6795,20 @@ void UEdGraphSchema_K2::SplitPin(UEdGraphPin* Pin, const bool bNotify) const
 
 void UEdGraphSchema_K2::RecombinePin(UEdGraphPin* Pin) const
 {
-	UEdGraphNode* GraphNode = Pin->GetOwningNode();
 	UEdGraphPin* ParentPin = Pin->ParentPin;
 
+	if (ParentPin == nullptr)
+	{
+		if (Pin->SubPins.Num() > 0)
+		{
+			RecombinePin(Pin->SubPins[0]);
+		}
+
+		return;
+	}
+
+	UEdGraphNode* GraphNode = Pin->GetOwningNode();
+
 	GraphNode->Modify();
 	ParentPin->Modify();
 
diff --git a/Engine/Source/Editor/BlueprintGraph/Private/K2Node.cpp b/Engine/Source/Editor/BlueprintGraph/Private/K2Node.cpp
--- a/Engine/Source/Editor/BlueprintGraph/Private/K2Node.cpp
+++ b/Engine/Source/Editor/BlueprintGraph/Private/K2Node.cpp
@@ -843,6 +843,31 @@ void UK2Node::ValidateNodeDuringCompilation(FCompilerResultsLog& MessageLog) con
 	}
 }
 
+FString UK2Node::GetPinMetaData(FName InPinName, FName InKey)
+{
+	UEdGraphPin* Pin = FindPin(InPinName);
+
+	// For split pins check the struct's metadata
+	if (Pin && Pin->ParentPin && Pin->ParentPin->PinType.PinCategory == UEdGraphSchema_K2::PC_Struct)
+	{
+		FName NewPinPropertyName = FName(*Pin->PinName.ToString().RightChop(Pin->ParentPin->PinName.ToString().Len() + 1));
+
+		UStruct* StructType = Cast<UStruct>(Pin->ParentPin->PinType.PinSubCategoryObject.Get());
+		if (StructType)
+		{
+			for (TFieldIterator<UProperty> It(StructType); It; ++It)
+			{
+				const UProperty* Property = *It;
+				if (Property && Property->GetFName() == NewPinPropertyName)
+				{
+					return Property->GetMetaData(InKey);
+				}
+			}
+		}
+	}
+	return FString();
+}
+
 void UK2Node::RewireOldPinsToNewPins(TArray<UEdGraphPin*>& InOldPins, TArray<UEdGraphPin*>& InNewPins)
 {
 	TArray<UEdGraphPin*> OrphanedOldPins;
diff --git a/Engine/Source/Editor/BlueprintGraph/Private/K2Node_Composite.cpp b/Engine/Source/Editor/BlueprintGraph/Private/K2Node_Composite.cpp
--- a/Engine/Source/Editor/BlueprintGraph/Private/K2Node_Composite.cpp
+++ b/Engine/Source/Editor/BlueprintGraph/Private/K2Node_Composite.cpp
@@ -48,6 +48,8 @@ void UK2Node_Composite::AllocateDefaultPins()
 			}
 		}
 	}
+
+	CacheWildcardPins();
 }
 
 void UK2Node_Composite::DestroyNode()
diff --git a/Engine/Source/Editor/BlueprintGraph/Private/K2Node_FunctionEntry.cpp b/Engine/Source/Editor/BlueprintGraph/Private/K2Node_FunctionEntry.cpp
--- a/Engine/Source/Editor/BlueprintGraph/Private/K2Node_FunctionEntry.cpp
+++ b/Engine/Source/Editor/BlueprintGraph/Private/K2Node_FunctionEntry.cpp
@@ -214,11 +214,10 @@ void UK2Node_FunctionEntry::PreSave(const class ITargetPlatform* TargetPlatform)
 					{
 						const UStructProperty* PotentialUDSProperty = Cast<const UStructProperty>(Property);
 						// UDS requires default data even when the LocalVariable value is empty
-						const bool bUDSProperty = PotentialUDSProperty && Cast<const UUserDefinedStruct>(PotentialUDSProperty->Struct);
 
 						for (FBPVariableDescription& LocalVar : LocalVariables)
 						{
-							if (LocalVar.VarName == Property->GetFName() && (bUDSProperty || !LocalVar.DefaultValue.IsEmpty()))
+							if (LocalVar.VarName == Property->GetFName() && !LocalVar.DefaultValue.IsEmpty())
 							{
 								// Go to property and back, this handles redirector fixup and will sanitize the output
 								// The asset registry only knows about these references because when the node is expanded it turns into a hard reference
@@ -502,12 +501,10 @@ void UK2Node_FunctionEntry::ExpandNode(class FKismetCompilerContext& CompilerCon
 			if (const UProperty* Property = *It)
 			{
 				const UStructProperty* PotentialUDSProperty = Cast<const UStructProperty>(Property);
-				//UDS requires default data even when the LocalVariable value is empty
-				const bool bUDSProperty = PotentialUDSProperty && Cast<const UUserDefinedStruct>(PotentialUDSProperty->Struct);
 
 				for (const FBPVariableDescription& LocalVar : LocalVariables)
 				{
-					if (LocalVar.VarName == Property->GetFName() && (bUDSProperty || !LocalVar.DefaultValue.IsEmpty()))
+					if (LocalVar.VarName == Property->GetFName() && !LocalVar.DefaultValue.IsEmpty())
 					{
 						// Add a variable set node for the local variable and hook it up immediately following the entry node or the last added local variable
 						UK2Node_VariableSet* VariableSetNode = CompilerContext.SpawnIntermediateNode<UK2Node_VariableSet>(this, SourceGraph);
@@ -590,40 +587,9 @@ void UK2Node_FunctionEntry::ExpandNode(class FKismetCompilerContext& CompilerCon
 	}
 }
 
-static void RefreshUDSValuesStoredAsString(const FEdGraphPinType& VarType, FString& Value)
-{
-	// For container structs just keep the value from before, container structs do not delta serialize
-	if (!Value.IsEmpty() && VarType.PinCategory == UEdGraphSchema_K2::PC_Struct && !VarType.IsContainer())
-	{
-		UUserDefinedStruct* UDS = Cast<UUserDefinedStruct>(VarType.PinSubCategoryObject.Get());
-		if (UDS)
-		{
-			FStructOnScope StructInstance(UDS);
-			UDS->InitializeDefaultValue(StructInstance.GetStructMemory());
-			UDS->ImportText(*Value, StructInstance.GetStructMemory(), nullptr, PPF_None, GLog, FString());
-
-			Value.Reset();
-			FStructOnScope DefaultStructInstance(UDS);
-			UDS->InitializeDefaultValue(DefaultStructInstance.GetStructMemory());
-			UDS->ExportText(Value, StructInstance.GetStructMemory(), DefaultStructInstance.GetStructMemory(), nullptr, PPF_None, nullptr);
-		}
-	}
-}
-
 void UK2Node_FunctionEntry::PostReconstructNode()
 {
 	Super::PostReconstructNode();
-
-	UBlueprint* Blueprint = GetBlueprint();
-
-	// We want to refresh old UDS default values of local variables. It's enough to do this once.
-	if (Blueprint && Blueprint->bIsRegeneratingOnLoad)
-	{
-		for (FBPVariableDescription& LocalVariable : LocalVariables)
-		{
-			RefreshUDSValuesStoredAsString(LocalVariable.VarType, LocalVariable.DefaultValue);
-		}
-	}
 }
 
 bool UK2Node_FunctionEntry::ModifyUserDefinedPinDefaultValue(TSharedPtr<FUserPinInfo> PinInfo, const FString& NewDefaultValue)
diff --git a/Engine/Source/Editor/BlueprintGraph/Private/K2Node_MacroInstance.cpp b/Engine/Source/Editor/BlueprintGraph/Private/K2Node_MacroInstance.cpp
--- a/Engine/Source/Editor/BlueprintGraph/Private/K2Node_MacroInstance.cpp
+++ b/Engine/Source/Editor/BlueprintGraph/Private/K2Node_MacroInstance.cpp
@@ -123,6 +123,8 @@ void UK2Node_MacroInstance::AllocateDefaultPins()
 			}
 		}
 	}
+
+	CacheWildcardPins();
 }
 
 void UK2Node_MacroInstance::PreloadRequiredAssets()
@@ -442,6 +444,15 @@ void UK2Node_MacroInstance::PostFixupAllWildcardPins(bool bInAllWildcardPinsUnli
 	{
 		// Reset the type to a wildcard because there are no longer any wildcard pins linked to determine a type with
 		ResolvedWildcardType.ResetToDefaults();
+
+		// Collapse any wildcard pins that are split and set their type back to wildcard
+		for (UEdGraphPin* Pin : WildcardPins)
+		{
+			GetSchema()->RecombinePin(Pin);
+			Pin->PinType.PinCategory = UEdGraphSchema_K2::PC_Wildcard;
+			Pin->PinType.PinSubCategory = NAME_None;
+			Pin->PinType.PinSubCategoryObject = nullptr;
+		}
 	}
 }
 
diff --git a/Engine/Source/Editor/BlueprintGraph/Private/K2Node_StructOperation.cpp b/Engine/Source/Editor/BlueprintGraph/Private/K2Node_StructOperation.cpp
--- a/Engine/Source/Editor/BlueprintGraph/Private/K2Node_StructOperation.cpp
+++ b/Engine/Source/Editor/BlueprintGraph/Private/K2Node_StructOperation.cpp
@@ -85,18 +85,15 @@ bool UK2Node_StructOperation::DoRenamedPinsMatch(const UEdGraphPin* NewPin, cons
 
 FString UK2Node_StructOperation::GetPinMetaData(FName InPinName, FName InKey)
 {
-	FString ReturnValue;
-
 	for (TFieldIterator<UProperty> It(StructType); It; ++It)
 	{
 		const UProperty* Property = *It;
 		if(Property && Property->GetFName() == InPinName)
 		{
-			ReturnValue = Property->GetMetaData(InKey);
-			break;
+			return Property->GetMetaData(InKey);
 		}
 	}
-	return ReturnValue;
+	return Super::GetPinMetaData(InPinName, InKey);
 }
 
 FString UK2Node_StructOperation::GetFindReferenceSearchString() const
diff --git a/Engine/Source/Editor/BlueprintGraph/Private/K2Node_Tunnel.cpp b/Engine/Source/Editor/BlueprintGraph/Private/K2Node_Tunnel.cpp
--- a/Engine/Source/Editor/BlueprintGraph/Private/K2Node_Tunnel.cpp
+++ b/Engine/Source/Editor/BlueprintGraph/Private/K2Node_Tunnel.cpp
@@ -264,6 +264,20 @@ UObject* UK2Node_Tunnel::GetJumpTargetForDoubleClick() const
 	return TargetNode;
 }
 
+void UK2Node_Tunnel::CacheWildcardPins()
+{
+	WildcardPins.Reset();
+
+	for (UEdGraphPin* Pin : Pins)
+	{
+		// for each of the wildcard pins...
+		if (Pin->PinType.PinCategory == UEdGraphSchema_K2::PC_Wildcard)
+		{
+			WildcardPins.Add(Pin);
+		}
+	}
+}
+
 void UK2Node_Tunnel::ReallocatePinsDuringReconstruction(TArray<UEdGraphPin*>& OldPins)
 {
 	Super::ReallocatePinsDuringReconstruction(OldPins);
@@ -271,44 +285,57 @@ void UK2Node_Tunnel::ReallocatePinsDuringReconstruction(TArray<UEdGraphPin*>& Ol
 	// determine if all wildcard pins are unlinked.
 	// if they are, we should revert them all back to wildcard status
 	bool bAllWildcardsAreUnlinked = true;
-	for (auto PinIt = Pins.CreateConstIterator(); PinIt; PinIt++)
+	for (UEdGraphPin* Pin : WildcardPins)
 	{
-		// for each of the wildcard pins...
-		UEdGraphPin* const Pin = *PinIt;
-		if (Pin->PinType.PinCategory == UEdGraphSchema_K2::PC_Wildcard)
+		// find it in the old pins array (where it might not be a wildcard)
+		// and see if it's unlinked
+		for (UEdGraphPin* OldPin : OldPins)
 		{
-			// find it in the old pins array (where it might not be a wildcard)
-			// and see if it's unlinked
-			for (auto OldPinIt = OldPins.CreateConstIterator(); OldPinIt; OldPinIt++)
+			if (OldPin->PinName == Pin->PinName)
 			{
-				UEdGraphPin const* const OldPin = *OldPinIt;
-				if (OldPin->PinName == Pin->PinName)
+				TFunction<bool(UEdGraphPin*)> IsPinLinked = [&IsPinLinked](UEdGraphPin* InPin)
 				{
-					if (OldPin->LinkedTo.Num() > 0)
+					if (InPin->LinkedTo.Num() > 0)
 					{
-						bAllWildcardsAreUnlinked = false;
-						break;
+						return true;
+					}
+					for (UEdGraphPin* SubPin : InPin->SubPins)
+					{
+						if (IsPinLinked(SubPin))
+						{
+							return true;
+						}
 					}
+					return false;
+				};
+				if (IsPinLinked(OldPin))
+				{
+					bAllWildcardsAreUnlinked = false;
+					break;
 				}
 			}
 		}
+		if (bAllWildcardsAreUnlinked == false)
+		{
+			break;
+		}
 	}
 
 	if (bAllWildcardsAreUnlinked == false)
 	{
 		// Copy pin types from old pins for wildcard pins
-		for (auto PinIt = Pins.CreateConstIterator(); PinIt; PinIt++)
+		for (UEdGraphPin* const Pin : WildcardPins)
 		{
-			UEdGraphPin* const Pin = *PinIt;
+			// Only change the type if it is still a wildcard
 			if (Pin->PinType.PinCategory == UEdGraphSchema_K2::PC_Wildcard)
 			{
 				// find it in the old pins and copy the type
-				for (auto OldPinIt = OldPins.CreateConstIterator(); OldPinIt; OldPinIt++)
+				for (UEdGraphPin const* const OldPin : OldPins)
 				{
-					UEdGraphPin const* const OldPin = *OldPinIt;
 					if (OldPin->PinName == Pin->PinName)
 					{
 						Pin->PinType = OldPin->PinType;
+						break;
 					}
 				}
 			}
diff --git a/Engine/Source/Editor/GraphEditor/Private/BlueprintConnectionDrawingPolicy.cpp b/Engine/Source/Editor/GraphEditor/Private/BlueprintConnectionDrawingPolicy.cpp
--- a/Engine/Source/Editor/GraphEditor/Private/BlueprintConnectionDrawingPolicy.cpp
+++ b/Engine/Source/Editor/GraphEditor/Private/BlueprintConnectionDrawingPolicy.cpp
@@ -13,54 +13,6 @@
 #include "Kismet2/BlueprintEditorUtils.h"
 #include "BlueprintEditorSettings.h"
 
-namespace
-{
-	// Attempt to find a tunnel instance node in the given source graph that matches the given target graph.
-	// @TODO - This is a potential performance issue as it's a linear/recursive search. We used to map nested
-	// Macro Instances to code offsets through the debug data, and search only the mapped set, but the solution
-	// only worked for Macro Instance nodes. Consider reusing that solution and extending it to include Composites.
-	static UEdGraphNode* FindMatchingTunnelInstanceNode(const UEdGraph* InGraph, const UEdGraph* TargetGraph)
-	{
-		// Gather up all tunnel nodes in the given source graph.
-		TArray<UK2Node_Tunnel*> TunnelNodes;
-		InGraph->GetNodesOfClass(TunnelNodes);
-		for (UEdGraphNode* TunnelNode : TunnelNodes)
-		{
-			// Check all tunnel instance nodes for a match against the associated tunnel instance graph.
-			if (UK2Node_MacroInstance* MacroInstanceNode = Cast<UK2Node_MacroInstance>(TunnelNode))
-			{
-				const UEdGraph* MacroGraph = MacroInstanceNode->GetMacroGraph();
-				if (TargetGraph == MacroGraph)
-				{
-					// The current macro source graph matches up with the target graph context.
-					return MacroInstanceNode;
-				}
-				else if(UEdGraphNode* TunnelInstanceNode = FindMatchingTunnelInstanceNode(MacroGraph, TargetGraph))
-				{
-					// We recursively found a tunnel instance node nested inside the macro source graph that matches up with the target graph context.
-					return TunnelInstanceNode;
-				}
-			}
-			else if (UK2Node_Composite* CompositeNode = Cast<UK2Node_Composite>(TunnelNode))
-			{
-				if (TargetGraph == CompositeNode->BoundGraph)
-				{
-					// The current composite source graph matches up with the target graph context.
-					return CompositeNode;
-				}
-				else if(UEdGraphNode* TunnelInstanceNode = FindMatchingTunnelInstanceNode(CompositeNode->BoundGraph, TargetGraph))
-				{
-					// We recursively found a tunnel instance node nested inside the composite source graph that matches up with the target graph context.
-					return TunnelInstanceNode;
-				}
-			}
-		}
-
-		// No matching tunnel instance node could be found within the given source graph.
-		return nullptr;
-	}
-};
-
 /////////////////////////////////////////////////////
 // FKismetConnectionDrawingPolicy
 
@@ -112,10 +64,48 @@ void FKismetConnectionDrawingPolicy::Draw(TMap<TSharedRef<SWidget>, FArrangedWid
 	FConnectionDrawingPolicy::Draw(InPinGeometries, ArrangedNodes);
 }
 
-bool FKismetConnectionDrawingPolicy::CanBuildRoadmap() const
+UBlueprint* FKismetConnectionDrawingPolicy::GetTargetBlueprint() const
 {
-	UObject* ActiveObject = nullptr;
+	// Start out with the current graph Blueprint context.
 	if (UBlueprint* TargetBP = FBlueprintEditorUtils::FindBlueprintForGraph(GraphObj))
+	{
+		// Macro library Blueprints have no associated "active" object debugging context, so we need to determine which one to use from the call stack.
+		if (TargetBP->BlueprintType == BPTYPE_MacroLibrary)
+		{
+			// Walk backwards through the stack trace (check the most recent sample first) until we find a Blueprint context with the expansion and a valid "active" object for debugging.
+			const TSimpleRingBuffer<FKismetTraceSample>& TraceStack = FKismetDebugUtilities::GetTraceStack();
+			for (int32 i = TraceStack.Num() - 1; i >= 0 && !TargetBP->GetObjectBeingDebugged(); --i)
+			{
+				const FKismetTraceSample& Sample = TraceStack(i);
+				if (UObject* TestObject = Sample.Context.Get())
+				{
+					if (UBlueprintGeneratedClass* TargetClass = Cast<UBlueprintGeneratedClass>(TestObject->GetClass()))
+					{
+						// Check to see if the current instruction maps back to the Macro library source graph; if it does, this Blueprint contains an expansion of the Macro source graph.
+						const FBlueprintDebugData& DebugData = TargetClass->GetDebugData();
+						if (UEdGraphNode* Node = DebugData.FindSourceNodeFromCodeLocation(Sample.Function.Get(), Sample.Offset, /*bAllowImpreciseHit=*/ false))
+						{
+							if (GraphObj == Node->GetGraph())
+							{
+								// Switch the target Blueprint to the context containing the expansion of the current Macro source graph in the Macro library.
+								TargetBP = CastChecked<UBlueprint>(TargetClass->ClassGeneratedBy);
+							}
+						}
+					}
+				}
+			}
+		}
+
+		return TargetBP;
+	}
+
+	return nullptr;
+}
+
+bool FKismetConnectionDrawingPolicy::CanBuildRoadmap(UBlueprint* TargetBP) const
+{
+	UObject* ActiveObject = nullptr;
+	if (TargetBP)
 	{
 		ActiveObject = TargetBP->GetObjectBeingDebugged();
 	}
@@ -123,26 +113,26 @@ bool FKismetConnectionDrawingPolicy::CanBuildRoadmap() const
 	return ActiveObject != nullptr;
 }
 
+bool FKismetConnectionDrawingPolicy::CanBuildRoadmap() const
+{
+	return CanBuildRoadmap(GetTargetBlueprint());
+}
+
 void FKismetConnectionDrawingPolicy::BuildExecutionRoadmap()
 {
 	LatestTimeDiscovered = 0.0;
 
+	UBlueprint* TargetBP = GetTargetBlueprint();
+	
 	// Only do highlighting in PIE or SIE
-	if (!CanBuildRoadmap())
+	if (!CanBuildRoadmap(TargetBP))
 	{
 		return;
 	}
 
-	UBlueprint* TargetBP = FBlueprintEditorUtils::FindBlueprintForGraphChecked(GraphObj);
 	UObject* ActiveObject = TargetBP->GetObjectBeingDebugged();
 	check(ActiveObject); // Due to CanBuildRoadmap
 
-	// Redirect the target Blueprint when debugging with a macro graph visible
-	if (TargetBP->BlueprintType == BPTYPE_MacroLibrary)
-	{
-		TargetBP = Cast<UBlueprint>(ActiveObject->GetClass()->ClassGeneratedBy);
-	}
-
 	TArray<UEdGraphNode*> SequentialNodesInGraph;
 	TArray<double> SequentialNodeTimes;
 	TArray<UEdGraphPin*> SequentialExecPinsInGraph;
@@ -171,11 +161,23 @@ void FKismetConnectionDrawingPolicy::BuildExecutionRoadmap()
 							SequentialNodeTimes.Add(Sample.ObservationTime);
 							SequentialExecPinsInGraph.Add(AssociatedPin);
 						}
-						else if (UEdGraphNode* TunnelInstanceNode = FindMatchingTunnelInstanceNode(GraphObj, Node->GetGraph()))
+						else if (const TArray<TWeakObjectPtr<UEdGraphNode> >* ExpansionSourceNodes = DebugData.FindExpansionSourceNodesFromCodeLocation(Sample.Function.Get(), Sample.Offset))
 						{
-							SequentialNodesInGraph.Add(TunnelInstanceNode);
-							SequentialNodeTimes.Add(Sample.ObservationTime);
-							SequentialExecPinsInGraph.Add(AssociatedPin);
+							// Attempt to find an outer expansion (tunnel instance) source node that may have also been mapped to the sample code offset.
+							for (TWeakObjectPtr<UEdGraphNode> ExpansionSourceNodePtr : *ExpansionSourceNodes)
+							{
+								if (UEdGraphNode* TunnelInstanceNode = ExpansionSourceNodePtr.Get())
+								{
+									if (GraphObj == TunnelInstanceNode->GetGraph())
+									{
+										SequentialNodesInGraph.Add(TunnelInstanceNode);
+										SequentialNodeTimes.Add(Sample.ObservationTime);
+										SequentialExecPinsInGraph.Add(AssociatedPin);
+
+										break;
+									}
+								}
+							}
 						}
 					}
 				}
diff --git a/Engine/Source/Editor/GraphEditor/Public/BlueprintConnectionDrawingPolicy.h b/Engine/Source/Editor/GraphEditor/Public/BlueprintConnectionDrawingPolicy.h
--- a/Engine/Source/Editor/GraphEditor/Public/BlueprintConnectionDrawingPolicy.h
+++ b/Engine/Source/Editor/GraphEditor/Public/BlueprintConnectionDrawingPolicy.h
@@ -9,6 +9,7 @@
 #include "ConnectionDrawingPolicy.h"
 
 class FSlateWindowElementList;
+class UBlueprint;
 class UEdGraph;
 
 /////////////////////////////////////////////////////
@@ -69,6 +70,13 @@ protected:
 protected:
 	// Should this wire be treated as an execution pin, querying the execution trace for connected nodes to draw it differently?
 	virtual bool TreatWireAsExecutionPin(UEdGraphPin* InputPin, UEdGraphPin* OutputPin) const;
+
+	// Determines the target Blueprint for querying the active debugging context.
+	UBlueprint* GetTargetBlueprint() const;
+
+	// Determines if the given Blueprint context has a valid object for debugging.
+	bool CanBuildRoadmap(UBlueprint* TargetBP) const;
+
 public:
 	FKismetConnectionDrawingPolicy(int32 InBackLayerID, int32 InFrontLayerID, float ZoomFactor, const FSlateRect& InClippingRect, FSlateWindowElementList& InDrawElements, UEdGraph* InGraphObj);
 
diff --git a/Engine/Source/Editor/Kismet/Private/BlueprintCompilationManager.cpp b/Engine/Source/Editor/Kismet/Private/BlueprintCompilationManager.cpp
--- a/Engine/Source/Editor/Kismet/Private/BlueprintCompilationManager.cpp
+++ b/Engine/Source/Editor/Kismet/Private/BlueprintCompilationManager.cpp
@@ -804,9 +804,13 @@ void FBlueprintCompilationManagerImpl::FlushCompilationQueueImpl(TArray<UObject*
 						BP->GeneratedClass->ClassDefaultObject == nullptr || 
 						BP->GeneratedClass->ClassDefaultObject->GetClass() != BP->GeneratedClass);
 				// default value propagation occurs in ReinstaneBatch, CDO will be created via CompileFunctions call:
-				if(BP->GeneratedClass)
+				if(BP->ParentClass)
 				{
-					BP->GeneratedClass->ClassDefaultObject = nullptr;
+					if(BP->GeneratedClass)
+					{
+						BP->GeneratedClass->ClassDefaultObject = nullptr;
+					}
+
 					// Reset the flag, so if the user tries to use PIE it will warn them if the BP did not compile
 					BP->bDisplayCompilePIEWarning = true;
 		
@@ -896,6 +900,11 @@ void FBlueprintCompilationManagerImpl::FlushCompilationQueueImpl(TArray<UObject*
 				}
 			}
 
+			if(BP->GeneratedClass)
+			{
+				BP->GeneratedClass->SetUpRuntimeReplicationData();
+			}
+
 			ensure(BP->GeneratedClass == nullptr || BP->GeneratedClass->ClassDefaultObject->GetClass() == *(BP->GeneratedClass));
 		}
 	} // end GTimeCompiling scope
diff --git a/Engine/Source/Editor/Kismet/Private/BlueprintDetailsCustomization.cpp b/Engine/Source/Editor/Kismet/Private/BlueprintDetailsCustomization.cpp
--- a/Engine/Source/Editor/Kismet/Private/BlueprintDetailsCustomization.cpp
+++ b/Engine/Source/Editor/Kismet/Private/BlueprintDetailsCustomization.cpp
@@ -839,16 +839,14 @@ void FBlueprintVarActionDetails::CustomizeDetails( IDetailLayoutBuilder& DetailL
 				check(EntryNodes.Num() > 0);
 
 				const UStructProperty* PotentialUDSProperty = Cast<const UStructProperty>(VariableProperty);
-				//UDS requires default data even when the LocalVariable value is empty
-				const bool bUDSProperty = PotentialUDSProperty && Cast<const UUserDefinedStruct>(PotentialUDSProperty->Struct);
-
+				
 				UK2Node_FunctionEntry* FuncEntry = EntryNodes[0];
 				for (const FBPVariableDescription& LocalVar : FuncEntry->LocalVariables)
 				{
 					if(LocalVar.VarName == VariableProperty->GetFName()) //Property->GetFName())
 					{
 						// Only set the default value if there is one
-						if(bUDSProperty || !LocalVar.DefaultValue.IsEmpty())
+						if(!LocalVar.DefaultValue.IsEmpty())
 						{
 							FBlueprintEditorUtils::PropertyValueFromString(VariableProperty, LocalVar.DefaultValue, StructData->GetStructMemory());
 						}
diff --git a/Engine/Source/Editor/Kismet/Private/UserDefinedStructureEditor.cpp b/Engine/Source/Editor/Kismet/Private/UserDefinedStructureEditor.cpp
--- a/Engine/Source/Editor/Kismet/Private/UserDefinedStructureEditor.cpp
+++ b/Engine/Source/Editor/Kismet/Private/UserDefinedStructureEditor.cpp
@@ -103,7 +103,7 @@ public:
 	void Initialize()
 	{
 		StructData = MakeShareable(new FStructOnScope(UserDefinedStruct.Get()));
-		FStructureEditorUtils::Fill_MakeStructureDefaultValue(UserDefinedStruct.Get(), StructData->GetStructMemory());
+		UserDefinedStruct.Get()->InitializeDefaultValue(StructData->GetStructMemory());
 		StructData->SetPackage(UserDefinedStruct->GetOutermost());
 
 		FPropertyEditorModule& PropertyModule = FModuleManager::LoadModuleChecked<FPropertyEditorModule>("PropertyEditor");
@@ -155,7 +155,7 @@ public:
 			DetailsView->SetObject(UserDefinedStruct.Get());
 		}
 
-		FStructureEditorUtils::Fill_MakeStructureDefaultValue(UserDefinedStruct.Get(), StructData->GetStructMemory());
+		UserDefinedStruct.Get()->InitializeDefaultValue(StructData->GetStructMemory());
 	}
 
 	// FNotifyHook interface
diff --git a/Engine/Source/Editor/Kismet/Private/WorkflowOrientedApp/WorkflowTabManager.cpp b/Engine/Source/Editor/Kismet/Private/WorkflowOrientedApp/WorkflowTabManager.cpp
--- a/Engine/Source/Editor/Kismet/Private/WorkflowOrientedApp/WorkflowTabManager.cpp
+++ b/Engine/Source/Editor/Kismet/Private/WorkflowOrientedApp/WorkflowTabManager.cpp
@@ -47,7 +47,7 @@ bool FTabInfo::PayloadMatches(TSharedPtr<FTabPayload> A, TSharedPtr<FTabPayload>
 	}
 }
 
-void FTabInfo::AddTabHistory(TSharedPtr< struct FGenericTabHistory > InHistoryNode, bool bInSaveHistory/* = true*/, bool bSpawnNewTab/* = true*/)
+void FTabInfo::AddTabHistory(TSharedPtr< struct FGenericTabHistory > InHistoryNode, bool bInSaveHistory/* = true*/, bool bPrevTabMatches/* = false*/)
 {
 	// If the tab is not new, save the current history.
 	if(CurrentHistoryIndex >= 0 && bInSaveHistory)
@@ -75,7 +75,7 @@ void FTabInfo::AddTabHistory(TSharedPtr< struct FGenericTabHistory > InHistoryNo
 	CurrentHistoryIndex = History.Num() - 1;
 
 	// Evoke the history
-	InHistoryNode->EvokeHistory(AsShared(), bSpawnNewTab);
+	InHistoryNode->EvokeHistory(AsShared(), bPrevTabMatches);
 	InHistoryNode->GetFactory().Pin()->OnTabActivated(Tab.Pin());
 }
 
@@ -496,13 +496,17 @@ TSharedPtr<SDockTab> FDocumentTracker::OpenDocument(TSharedPtr<FTabPayload> InPa
 	}
 	else if(InOpenCause == OpenNewDocument || InOpenCause == CreateHistoryEvent)
 	{
+		TSharedPtr< FTabInfo > TabInfo;
 		TSharedPtr<SDockTab> Tab;
+		bool bPrevTabMatches = false;
 
 		// If the current tab matches we'll re-use it.
 		TSharedPtr<FTabInfo> LastEditedTabInfoPinned = LastEditedTabInfo.Pin();
 		if(LastEditedTabInfoPinned.IsValid() && LastEditedTabInfoPinned->PayloadMatches(InPayload))
 		{
-			Tab = LastEditedTabInfoPinned->GetTab().Pin();
+			TabInfo = LastEditedTabInfoPinned;
+			Tab = TabInfo->GetTab().Pin();
+			bPrevTabMatches = true;
 		}
 		else
 		{
@@ -512,13 +516,14 @@ TSharedPtr<SDockTab> FDocumentTracker::OpenDocument(TSharedPtr<FTabPayload> InPa
 			{
 				if ((*ListIt)->PayloadMatches(InPayload))
 				{
-					Tab = (*ListIt)->GetTab().Pin();
-					TabManager->DrawAttention( Tab.ToSharedRef() );
+					TabInfo = (*ListIt);
+					Tab = TabInfo->GetTab().Pin();
+					TabManager->DrawAttention( TabInfo->GetTab().Pin().ToSharedRef() );
 				}
 			}
 		}
 		
-		if(!Tab.IsValid())
+		if(!TabInfo.IsValid())
 		{
 			// If no tab was found with the payload
 			Tab = OpenNewTab(InPayload, OpenNewDocument);
@@ -526,7 +531,7 @@ TSharedPtr<SDockTab> FDocumentTracker::OpenDocument(TSharedPtr<FTabPayload> InPa
 		else if(InOpenCause == CreateHistoryEvent)
 		{
 			TSharedPtr<FDocumentTabFactory> Factory = FindSupportingFactory(InPayload.ToSharedRef());
-			LastEditedTabInfoPinned->AddTabHistory(Factory->CreateTabHistoryNode(InPayload), true, true);
+			TabInfo->AddTabHistory(Factory->CreateTabHistoryNode(InPayload), true, true);
 		}
 		
 		return Tab;
diff --git a/Engine/Source/Editor/KismetCompiler/Private/KismetCompiler.cpp b/Engine/Source/Editor/KismetCompiler/Private/KismetCompiler.cpp
--- a/Engine/Source/Editor/KismetCompiler/Private/KismetCompiler.cpp
+++ b/Engine/Source/Editor/KismetCompiler/Private/KismetCompiler.cpp
@@ -2286,6 +2286,7 @@ void FKismetCompilerContext::FinishCompilingClass(UClass* Class)
 	}
 
 	Class->StaticLink(true);
+	Class->SetUpRuntimeReplicationData();
 
 	// Create the default object for this class
 	FKismetCompilerUtilities::CompileDefaultProperties(Class);
@@ -3384,19 +3385,44 @@ void FKismetCompilerContext::ExpandTunnelsAndMacros(UEdGraph* SourceGraph)
 		{
 			UK2Node_Tunnel* InputSink = TunnelNode->GetInputSink();
 			UK2Node_Tunnel* OutputSource = TunnelNode->GetOutputSource();
-
-			if (FBlueprintEditorUtils::IsTunnelInstanceNode(OutputSource))
+			
+			// Determine the tunnel nodes that bound the expansion
+			UK2Node_Tunnel* TunnelInstance = nullptr;
+			UK2Node_Tunnel* TunnelInputSite = nullptr;
+			UK2Node_Tunnel *TunnelOutputSite = nullptr;
+			if (FBlueprintEditorUtils::IsTunnelInstanceNode(TunnelNode))
 			{
-				// Process the input side of a tunnel instance expansion
-				ProcessIntermediateTunnelBoundary(OutputSource, TunnelNode);
+				TunnelInstance = TunnelNode;
+				TunnelInputSite = InputSink;
+				TunnelOutputSite = OutputSource;
 			}
 			else if (FBlueprintEditorUtils::IsTunnelInstanceNode(InputSink))
 			{
-				// Process the output side of a tunnel instance expansion
-				ProcessIntermediateTunnelBoundary(TunnelNode, InputSink);
+				TunnelInstance = InputSink;
+				TunnelOutputSite = TunnelNode;
+			}
+			else if (FBlueprintEditorUtils::IsTunnelInstanceNode(OutputSource))
+			{
+				TunnelInstance = OutputSource;
+				TunnelInputSite = TunnelNode;
 			}
 
-			bool bSuccess = Schema->CollapseGatewayNode(TunnelNode, InputSink, OutputSource, this);
+			if (TunnelInstance)
+			{
+				if (TunnelInputSite)
+				{
+					// Construct an intermediate tunnel boundary on the input side of a tunnel instance expansion.
+					ProcessIntermediateTunnelBoundary(TunnelInstance, TunnelInputSite);
+				}
+
+				if (TunnelOutputSite)
+				{
+					// Construct an intermediate tunnel boundary on the output side of a tunnel instance expansion.
+					ProcessIntermediateTunnelBoundary(TunnelOutputSite, TunnelInstance);
+				}
+			}
+
+			const bool bSuccess = Schema->CollapseGatewayNode(TunnelNode, InputSink, OutputSource, this);
 			if (!bSuccess)
 			{
 				MessageLog.Error(*LOCTEXT("CollapseTunnel_Error", "Failed to collapse tunnel @@").ToString(), TunnelNode);
@@ -3936,8 +3962,6 @@ void FKismetCompilerContext::CompileFunctions(EInternalCompilerFlags InternalFla
 
 		UObject* NewCDO = NewClass->GetDefaultObject();
 
-		FUserDefinedStructureCompilerUtils::DefaultUserDefinedStructs(NewCDO, MessageLog);
-
 		// Copy over the CDO properties if we're not already regenerating on load.  In that case, the copy will be done after compile on load is complete
 		FBlueprintEditorUtils::PropagateParentBlueprintDefaults(NewClass);
 
@@ -4490,21 +4514,26 @@ void FKismetCompilerContext::MapExpansionPathToTunnelInstance(UEdGraphNode* Inne
 {
 	if (InnerExpansionNode && OuterTunnelInstance)
 	{
-		MessageLog.NotifyIntermediateTunnelNode(InnerExpansionNode, OuterTunnelInstance);
+		// Only map the node to the tunnel instance if it hasn't been mapped before (e.g. by a nested expansion).
+		if (!MessageLog.GetIntermediateTunnelInstance(InnerExpansionNode))
+		{
+			MessageLog.NotifyIntermediateTunnelNode(InnerExpansionNode, OuterTunnelInstance);
+		}
 
-		// Recursively map any nodes linked to this node along the execution path.
+		// Recursively map any nodes linked to this node along each output execution path.
 		for (const UEdGraphPin* OutputPin : InnerExpansionNode->Pins)
 		{
 			if (OutputPin->Direction == EGPD_Output && UEdGraphSchema_K2::IsExecPin(*OutputPin) && OutputPin->LinkedTo.Num() > 0)
 			{
 				for (const UEdGraphPin* LinkedTo : OutputPin->LinkedTo)
 				{
-					// Make sure it hasn't already been mapped (e.g. re-entrant execution paths).
-					UEdGraphNode* ExpansionNode = LinkedTo->GetOwningNode();
-					UEdGraphNode* TunnelInstance = MessageLog.GetIntermediateTunnelInstance(ExpansionNode);
-					if (TunnelInstance != OuterTunnelInstance)
+					// Make sure it is valid and hasn't already been mapped (e.g. shared execution paths). Also, avoid mapping tunnel output nodes (not needed).
+					UEdGraphNode* LinkedExpansionNode = LinkedTo->GetOwningNode();
+					if (LinkedExpansionNode
+						&& !MessageLog.GetIntermediateTunnelInstance(LinkedExpansionNode)
+						&& (!LinkedExpansionNode->IsA<UK2Node_Tunnel>() || FBlueprintEditorUtils::IsTunnelInstanceNode(LinkedExpansionNode)))
 					{
-						MapExpansionPathToTunnelInstance(ExpansionNode, OuterTunnelInstance);
+						MapExpansionPathToTunnelInstance(LinkedExpansionNode, OuterTunnelInstance);
 					}
 				}
 			}
@@ -4512,6 +4541,44 @@ void FKismetCompilerContext::MapExpansionPathToTunnelInstance(UEdGraphNode* Inne
 	}
 }
 
+// This method injects an intermediate "boundary" node on either side of a tunnel instance node and the tunnel input/output
+// nodes which can be found along the execution path that flows through the tunnel instance node's expansion. The boundary
+// nodes resolve to a NOP debug site for breakpoints and wire traces, and are only constructed when debug data is enabled.
+//
+// For example:
+//
+//     +======================+
+//     | Tunnel instance node |
+//     +======================+
+// (1) | >--+            +--> |
+//     +====|============|====+
+//          |            |
+//          |            +-------------------------------------------------------+
+//          |                                                                    |
+//          |   +================+                         +=================+   |
+//          |   | Input (Tunnel) |                         | Output (Tunnel) |   |
+//          |   +================+                         +=================+   |
+//          +---|--------------> | (2) . . . . . . . . (3) | >---------------|---+
+//              +================+                         +=================+
+//
+// In the expansion shown above, intermediate boundary nodes are created at the following locations along the execution path:
+//
+//	(1) "Entry" site - Precedes the tunnel instance node in the execution sequence.
+//	(2) "Input" site - Follows the input tunnel in the expansion of the tunnel instance.
+//	(3) "Output" site - Precedes the output tunnel in the expansion of the tunnel instance.
+//
+// After tunnels are collapsed and isolated in the intermediate function graph during expansion, the tunnel boundary nodes will
+// remain in place along the execution path, and they won't get compiled out. The resulting bytecode resolves to a NOP sequence.
+//
+// When a tunnel instance node has multiple exec inputs/outputs, this method creates one tunnel boundary per exec path through
+// the expansion. Also, note that we do not create a boundary node on the output side of the tunnel instance node, because we want
+// execution to continue on to the next linked node after the instruction pointer passes the tunnel output site when single-stepping.
+//
+// In addition to creating intermediate tunnel boundary nodes, this method also maps the intermediate impure nodes along each unique
+// execution path through the expansion (between boundaries 2 and 3 in the diagram above) back to the intermediate tunnel instance
+// node that resulted in the expansion. This mapping is used for (a) producing stable UUIDs for latent nodes in an expansion, and
+// (b) drawing "marching ants" on either side of the tunnel instance node that corresponds to the execution path in the source graph. 
+//
 void FKismetCompilerContext::ProcessIntermediateTunnelBoundary(UK2Node_Tunnel* TunnelInput, UK2Node_Tunnel* TunnelOutput)
 {
 	// @TODO move this check out of KismetFunctionContext so we can use it here?
@@ -4526,14 +4593,15 @@ void FKismetCompilerContext::ProcessIntermediateTunnelBoundary(UK2Node_Tunnel* T
 		// Set the base node name and boundary type.
 		TunnelBoundary->SetNodeAttributes(TunnelSource);
 
-		// Auto-position the node in the intermediate graph.
-		// @TODO - Maybe we need to find a better way to reposition these?
-		AutoAssignNodePosition(TunnelBoundary);
+		// Position the node in the intermediate graph.
+		TunnelBoundary->NodePosX = TunnelSource->NodePosX;
+		TunnelBoundary->NodePosY = TunnelSource->NodePosY;
 	};
 
 	if (TunnelInput)
 	{
-		const bool bIsInputBoundary = FBlueprintEditorUtils::IsTunnelInstanceNode(TunnelInput);
+		// Flag that indicates whether or not the tunnel instance node is designated as an input or an output.
+		const bool bIsTunnelEntrySite = FBlueprintEditorUtils::IsTunnelInstanceNode(TunnelInput);
 
 		for (UEdGraphPin* InputPin : TunnelInput->Pins)
 		{
@@ -4547,6 +4615,9 @@ void FKismetCompilerContext::ProcessIntermediateTunnelBoundary(UK2Node_Tunnel* T
 					{
 						InitializeTunnelBoundaryNode(InputBoundaryNode, TunnelInput);
 
+						// Map the intermediate input tunnel boundary node back to the intermediate tunnel instance node that spawned it.
+						MessageLog.NotifyIntermediateTunnelNode(InputBoundaryNode, bIsTunnelEntrySite ? TunnelInput : TunnelOutput);
+
 						if (UEdGraphPin* NewInputPin = InputBoundaryNode->CreatePin(EGPD_Input, InputPin->PinType, InputPin->PinName))
 						{
 							if (UEdGraphPin* NewOutputPin = InputBoundaryNode->CreatePin(EGPD_Output, InputPin->PinType, InputBoundaryNode->CreateUniquePinName(InputPin->PinName)))
@@ -4566,7 +4637,7 @@ void FKismetCompilerContext::ProcessIntermediateTunnelBoundary(UK2Node_Tunnel* T
 				{
 					if (ensure(OutputPin->Direction == EGPD_Output && UEdGraphSchema_K2::IsExecPin(*OutputPin)) && OutputPin->LinkedTo.Num() > 0)
 					{
-						if (bIsInputBoundary)
+						if (bIsTunnelEntrySite)
 						{
 							// Map the execution path through the expansion back to the tunnel instance node. Note that the assumption here is
 							// that we haven't collapsed the tunnels yet, so the output side of the expansion shouldn't be linked to anything.
@@ -4582,6 +4653,9 @@ void FKismetCompilerContext::ProcessIntermediateTunnelBoundary(UK2Node_Tunnel* T
 								{
 									InitializeTunnelBoundaryNode(OutputBoundaryNode, TunnelOutput);
 
+									// Map the intermediate output tunnel boundary node back to the intermediate tunnel instance node that spawned it.
+									MessageLog.NotifyIntermediateTunnelNode(OutputBoundaryNode, TunnelInput);
+
 									if (UEdGraphPin* NewInputPin = OutputBoundaryNode->CreatePin(EGPD_Input, OutputPin->PinType, OutputPin->PinName))
 									{
 										if (UEdGraphPin* NewOutputPin = OutputBoundaryNode->CreatePin(EGPD_Output, OutputPin->PinType, OutputBoundaryNode->CreateUniquePinName(OutputPin->PinName)))
diff --git a/Engine/Source/Editor/KismetCompiler/Private/KismetCompilerVMBackend.cpp b/Engine/Source/Editor/KismetCompiler/Private/KismetCompilerVMBackend.cpp
--- a/Engine/Source/Editor/KismetCompiler/Private/KismetCompilerVMBackend.cpp
+++ b/Engine/Source/Editor/KismetCompiler/Private/KismetCompilerVMBackend.cpp
@@ -832,10 +832,6 @@ public:
 					{
 						UE_LOG(LogK2Compiler, Error, TEXT("Unsupported static array. Property: %s, Struct: %s"), *GetPathNameSafe(StructProperty), *Struct->GetName());
 					}
-					if(!FStructureEditorUtils::Fill_MakeStructureDefaultValue(Cast<UUserDefinedStruct>(Struct), StructData))
-					{
-						UE_LOG(LogK2Compiler, Warning, TEXT("MakeStructureDefaultValue parsing error. Property: %s, Struct: %s"), *GetPathNameSafe(StructProperty), *Struct->GetName());
-					}
 
 					// Assume that any errors on the import of the name string have been caught in the function call generation
 					Struct->ImportText(Term->Name.IsEmpty() ? TEXT("()") : *Term->Name, StructData, nullptr, PPF_None, GLog, GetPathNameSafe(StructProperty));
@@ -1859,8 +1855,45 @@ public:
 				const bool bInstrumentedBreakpoint = Statement.Type == KCST_InstrumentedWireEntry;
 				const bool bBreakpointSite = Statement.Type == KCST_DebugSite || bInstrumentedBreakpoint;
 
+				// If the intermediate node was the result of a tunnel expansion (e.g. macro/composite), gather the chain of tunnel instance
+				// source nodes that were expanded to eventually include the node. This information is used to construct a lookup table that's
+				// used to help determine whether or not the current instruction pointer maps back to a tunnel instance node in a source graph.
+				// For example:
+				//
+				//	A (Macro instance node in top-level source graph)
+				//	|
+				//	+- [...expansion of A...] + B (Composite node in A's macro source graph)
+				//	   |
+				//	   +- [...expansion of B...] + C (Macro instance node in B's collapsed child subgraph)
+				//	      |
+				//        +- [...expansion of C...]
+				//
+				//  The intermediate exec nodes in each expansion set will map back to C through MessageLog.GetIntermediateTunnelInstance().
+				//	Thus, if SourceNode was created as the result of a tunnel instance expansion, it will map back to an outer Tunnel Instance
+				//	node. If the Tunnel Instance node itself was created as the result of a tunnel instance expansion, it will map back to
+				//	another outer Tunnel Instance node as well. This will continue until we run out of Tunnel Instance nodes. The set of Tunnel
+				//	Instance nodes that we find constitutes the full expansion hierarchy. We then map the hierarchy back to their matching
+				//	source nodes and register the set as the Tunnel Instance node chain that's associated with the current instruction offset.
+				//
+				TArray<TWeakObjectPtr<UEdGraphNode>> ExpansionSourceNodes;
+				if (UEdGraphNode* OuterTunnelInstance = FunctionContext.MessageLog.GetIntermediateTunnelInstance(SourceNode))
+				{
+					do
+					{
+						// Map the intermediate tunnel instance node back to its original source and add it to the array.
+						if (UEdGraphNode* ExpansionSourceNode = Cast<UEdGraphNode>(FunctionContext.MessageLog.FindSourceObject(OuterTunnelInstance)))
+						{
+							ExpansionSourceNodes.Add(ExpansionSourceNode);
+						}
+						
+						// Continue back up the chain until we run out of expansion source nodes (this ensures that we include any nested expansions).
+						OuterTunnelInstance = FunctionContext.MessageLog.GetIntermediateTunnelInstance(OuterTunnelInstance);
+
+					} while (OuterTunnelInstance);
+				}
+
 				// Register the debug information for the node.
-				ClassBeingBuilt->GetDebugData().RegisterNodeToCodeAssociation(TrueSourceNode, FunctionContext.Function, Offset, bBreakpointSite);
+				ClassBeingBuilt->GetDebugData().RegisterNodeToCodeAssociation(TrueSourceNode, ExpansionSourceNodes, FunctionContext.Function, Offset, bBreakpointSite);
 
 				// Track pure node script code range for the current impure (exec) node
 				if (Statement.Type == KCST_InstrumentedPureNodeEntry)
diff --git a/Engine/Source/Editor/KismetCompiler/Private/UserDefinedStructureCompilerUtils.cpp b/Engine/Source/Editor/KismetCompiler/Private/UserDefinedStructureCompilerUtils.cpp
--- a/Engine/Source/Editor/KismetCompiler/Private/UserDefinedStructureCompilerUtils.cpp
+++ b/Engine/Source/Editor/KismetCompiler/Private/UserDefinedStructureCompilerUtils.cpp
@@ -20,17 +20,29 @@
 #include "KismetCompiler.h"
 #include "Kismet2/StructureEditorUtils.h"
 #include "Kismet2/BlueprintEditorUtils.h"
+#include "Serialization/ObjectWriter.h"
+#include "Serialization/ObjectReader.h"
 
 #define LOCTEXT_NAMESPACE "StructureCompiler"
 
 struct FUserDefinedStructureCompilerInner
 {
-	static void ClearStructReferencesInBP(UBlueprint* FoundBlueprint, TSet<UBlueprint*>& BlueprintsToRecompile)
+	struct FBlueprintUserStructData
 	{
-		bool bAlreadyProcessed = false;
-		BlueprintsToRecompile.Add(FoundBlueprint, &bAlreadyProcessed);
-		if (!bAlreadyProcessed)
+		TArray<uint8> SkeletonCDOData;
+		TArray<uint8> GeneratedCDOData;
+	};
+
+	static void ClearStructReferencesInBP(UBlueprint* FoundBlueprint, TMap<UBlueprint*, FBlueprintUserStructData>& BlueprintsToRecompile)
+	{
+		if (!BlueprintsToRecompile.Contains(FoundBlueprint))
 		{
+			FBlueprintUserStructData& BlueprintData = BlueprintsToRecompile.Add(FoundBlueprint);
+
+			// Write CDO data to temp archive
+			//FObjectWriter SkeletonMemoryWriter(FoundBlueprint->SkeletonGeneratedClass->GetDefaultObject(), BlueprintData.SkeletonCDOData);
+			//FObjectWriter MemoryWriter(FoundBlueprint->GeneratedClass->GetDefaultObject(), BlueprintData.GeneratedCDOData);
+
 			for (UFunction* Function : TFieldRange<UFunction>(FoundBlueprint->GeneratedClass, EFieldIteratorFlags::ExcludeSuper))
 			{
 				Function->Script.Empty();
@@ -41,7 +53,7 @@ struct FUserDefinedStructureCompilerInner
 
 	static void ReplaceStructWithTempDuplicate(
 		UUserDefinedStruct* StructureToReinstance, 
-		TSet<UBlueprint*>& BlueprintsToRecompile,
+		TMap<UBlueprint*, FBlueprintUserStructData>& BlueprintsToRecompile,
 		TArray<UUserDefinedStruct*>& ChangedStructs)
 	{
 		if (StructureToReinstance)
@@ -62,6 +74,7 @@ struct FUserDefinedStructureCompilerInner
 			DuplicatedStruct->Status = EUserDefinedStructureStatus::UDSS_Duplicate;
 			DuplicatedStruct->SetFlags(RF_Transient);
 			DuplicatedStruct->AddToRoot();
+
 			CastChecked<UUserDefinedStructEditorData>(DuplicatedStruct->EditorData)->RecreateDefaultInstance();
 
 			for (UStructProperty* StructProperty : TObjectRange<UStructProperty>(RF_ClassDefaultObject, /** bIncludeDerivedClasses */ true, /** InternalExcludeFlags */ EInternalObjectFlags::PendingKill))
@@ -72,8 +85,8 @@ struct FUserDefinedStructureCompilerInner
 					{
 						if (UBlueprint* FoundBlueprint = Cast<UBlueprint>(OwnerClass->ClassGeneratedBy))
 						{
-							StructProperty->Struct = DuplicatedStruct;
 							ClearStructReferencesInBP(FoundBlueprint, BlueprintsToRecompile);
+							StructProperty->Struct = DuplicatedStruct;
 						}
 					}
 					else if (UUserDefinedStruct* OwnerStruct = Cast<UUserDefinedStruct>(StructProperty->GetOwnerStruct()))
@@ -86,7 +99,12 @@ struct FUserDefinedStructureCompilerInner
 						if (bValidStruct)
 						{
 							ChangedStructs.AddUnique(OwnerStruct);
-							StructProperty->Struct = DuplicatedStruct;
+
+							if (FStructureEditorUtils::FStructEditorManager::ActiveChange != FStructureEditorUtils::EStructureEditorChangeInfo::DefaultValueChanged)
+							{
+								// Don't change this for a default value only change, it won't get correctly replaced later
+								StructProperty->Struct = DuplicatedStruct;
+							}
 						}
 					}
 					else
@@ -275,7 +293,6 @@ struct FUserDefinedStructureCompilerInner
 		Struct->SetMetaData(FBlueprintMetadata::MD_Tooltip, *FStructureEditorUtils::GetTooltip(Struct));
 
 		UUserDefinedStructEditorData* EditorData = CastChecked<UUserDefinedStructEditorData>(Struct->EditorData);
-		Struct->SetSuperStruct(EditorData->NativeBase);	
 
 		CreateVariables(Struct, K2Schema, MessageLog);
 
@@ -377,7 +394,7 @@ void FUserDefinedStructureCompilerUtils::CompileStruct(class UUserDefinedStruct*
 			ChangedStructs.Add(Struct);
 		}
 
-		TSet<UBlueprint*> BlueprintsToRecompile;
+		TMap<UBlueprint*, FUserDefinedStructureCompilerInner::FBlueprintUserStructData> BlueprintsToRecompile;
 		for (int32 StructIdx = 0; StructIdx < ChangedStructs.Num(); ++StructIdx)
 		{
 			UUserDefinedStruct* ChangedStruct = ChangedStructs[StructIdx];
@@ -441,6 +458,9 @@ void FUserDefinedStructureCompilerUtils::CompileStruct(class UUserDefinedStruct*
 					{
 						BlueprintsThatHaveBeenRecompiled.Add(FoundBlueprint);
 						BlueprintsToRecompile.Remove(FoundBlueprint);
+
+						// Reapply CDO data
+
 						FBlueprintEditorUtils::MarkBlueprintAsStructurallyModified(FoundBlueprint);
 					}
 					Node->ReconstructNode();
@@ -448,9 +468,9 @@ void FUserDefinedStructureCompilerUtils::CompileStruct(class UUserDefinedStruct*
 			}
 		}
 
-		for (UBlueprint* BlueprintToRecompile : BlueprintsToRecompile)
+		for (TPair<UBlueprint*, FUserDefinedStructureCompilerInner::FBlueprintUserStructData>& Pair : BlueprintsToRecompile)
 		{
-			FBlueprintEditorUtils::MarkBlueprintAsStructurallyModified(BlueprintToRecompile);
+			FBlueprintEditorUtils::MarkBlueprintAsStructurallyModified(Pair.Key);
 		}
 
 		for (UUserDefinedStruct* ChangedStruct : ChangedStructs)
@@ -464,23 +484,4 @@ void FUserDefinedStructureCompilerUtils::CompileStruct(class UUserDefinedStruct*
 	}
 }
 
-void FUserDefinedStructureCompilerUtils::DefaultUserDefinedStructs(UObject* Object, FCompilerResultsLog& MessageLog)
-{
-	if (Object && FStructureEditorUtils::UserDefinedStructEnabled())
-	{
-		for (TFieldIterator<UProperty> It(Object->GetClass()); It; ++It)
-		{
-			if (const UProperty* Property = (*It))
-			{
-				uint8* Mem = Property->ContainerPtrToValuePtr<uint8>(Object);
-				if (!FStructureEditorUtils::Fill_MakeStructureDefaultValue(Property, Mem))
-				{
-					MessageLog.Warning(*FString::Printf(*LOCTEXT("MakeStructureDefaultValue_Error", "MakeStructureDefaultValue parsing error. Object: %s, Property: %s").ToString(),
-						*Object->GetName(), *Property->GetName()));
-				}
-			}
-		}
-	}
-}
-
 #undef LOCTEXT_NAMESPACE
diff --git a/Engine/Source/Editor/KismetCompiler/Public/UserDefinedStructureCompilerUtils.h b/Engine/Source/Editor/KismetCompiler/Public/UserDefinedStructureCompilerUtils.h
--- a/Engine/Source/Editor/KismetCompiler/Public/UserDefinedStructureCompilerUtils.h
+++ b/Engine/Source/Editor/KismetCompiler/Public/UserDefinedStructureCompilerUtils.h
@@ -11,9 +11,5 @@ class FUserDefinedStructureCompilerUtils
 public:
 	// ASSUMPTION, structure doesn't need to be renamed or removed
 	static void CompileStruct(class UUserDefinedStruct* Struct, class FCompilerResultsLog& MessageLog, bool bForceRecompile);
-
-	// Fill all variables of given object, that are UserDefinedStructure type, with default values.
-	// Default values for member variables in User Defined Structure are stored in meta data "MakeStructureDefaultValue"
-	static void DefaultUserDefinedStructs(UObject* Object, FCompilerResultsLog& MessageLog);
 };
 
diff --git a/Engine/Source/Editor/PropertyEditor/Private/PropertyHandleImpl.cpp b/Engine/Source/Editor/PropertyEditor/Private/PropertyHandleImpl.cpp
--- a/Engine/Source/Editor/PropertyEditor/Private/PropertyHandleImpl.cpp
+++ b/Engine/Source/Editor/PropertyEditor/Private/PropertyHandleImpl.cpp
@@ -1268,7 +1268,6 @@ void FPropertyValueImpl::AddChild()
 					{
 						FScriptArrayHelper	ArrayHelper(Array, Addr);
 						Index = ArrayHelper.AddValue();
-						FPropertyNode::AdditionalInitializationUDS(Array->Inner, ArrayHelper.GetRawPtr(Index));
 					}
 					else if (Set)
 					{
@@ -1276,7 +1275,6 @@ void FPropertyValueImpl::AddChild()
 						Index = SetHelper.AddDefaultValue_Invalid_NeedsRehash();
 						SetHelper.Rehash();
 
-						FPropertyNode::AdditionalInitializationUDS(Set->ElementProp, SetHelper.GetElementPtr(Index));
 					}
 					else if (Map)
 					{
@@ -1285,8 +1283,6 @@ void FPropertyValueImpl::AddChild()
 						MapHelper.Rehash();
 
 						uint8* PairPtr = MapHelper.GetPairPtr(Index);
-						FPropertyNode::AdditionalInitializationUDS(Map->KeyProp, Map->KeyProp->ContainerPtrToValuePtr<uint8>(PairPtr));
-						FPropertyNode::AdditionalInitializationUDS(Map->ValueProp, Map->ValueProp->ContainerPtrToValuePtr<uint8>(PairPtr));
 					}
 
 					ArrayIndicesPerObject[i].Add(NodeProperty->GetName(), Index);
@@ -1532,8 +1528,6 @@ void FPropertyValueImpl::InsertChild( TSharedPtr<FPropertyNode> ChildNodeToInser
 
 		ArrayHelper.InsertValues(Index, 1 );
 
-		FPropertyNode::AdditionalInitializationUDS(ArrayProperty->Inner, ArrayHelper.GetRawPtr(Index));
-
 		//set up indices for the coming events
 		TArray< TMap<FString,int32> > ArrayIndicesPerObject;
 		for (int32 ObjectIndex = 0; ObjectIndex < ReadAddresses.Num(); ++ObjectIndex)
@@ -1859,8 +1853,6 @@ void FPropertyValueImpl::MoveElementTo(int32 OriginalIndex, int32 NewIndex)
 
 			ArrayHelper.InsertValues(Index, 1);
 
-			FPropertyNode::AdditionalInitializationUDS(ArrayProperty->Inner, ArrayHelper.GetRawPtr(Index));
-
 			//set up indices for the coming events
 			TArray< TMap<FString, int32> > ArrayIndicesPerObject;
 			for (int32 ObjectIndex = 0; ObjectIndex < ReadAddresses.Num(); ++ObjectIndex)
@@ -1933,8 +1925,6 @@ void FPropertyValueImpl::MoveElementTo(int32 OriginalIndex, int32 NewIndex)
 
 						FScriptArrayHelper	ArrayHelper(Array, Addr);
 						Index = ArrayHelper.AddValue();
-						FPropertyNode::AdditionalInitializationUDS(Array->Inner, ArrayHelper.GetRawPtr(Index));
-						
 
 						ArrayIndicesPerObject[i].Add(NodeProperty->GetName(), Index);
 					}
diff --git a/Engine/Source/Editor/PropertyEditor/Private/PropertyNode.cpp b/Engine/Source/Editor/PropertyEditor/Private/PropertyNode.cpp
--- a/Engine/Source/Editor/PropertyEditor/Private/PropertyNode.cpp
+++ b/Engine/Source/Editor/PropertyEditor/Private/PropertyNode.cpp
@@ -2679,17 +2679,6 @@ bool FPropertyNode::IsFilterAcceptable(const TArray<FString>& InAcceptableNames,
 	return bCompleteMatchFound;
 }
 
-void FPropertyNode::AdditionalInitializationUDS(UProperty* Property, uint8* RawPtr)
-{
-	if (const UStructProperty* StructProperty = Cast<const UStructProperty>(Property))
-	{
-		if (!FStructureEditorUtils::Fill_MakeStructureDefaultValue(Cast<const UUserDefinedStruct>(StructProperty->Struct), RawPtr))
-		{
-			UE_LOG(LogPropertyNode, Warning, TEXT("MakeStructureDefaultValue parsing error. Property: %s "), *StructProperty->GetPathName());
-		}
-	}
-}
-
 void FPropertyNode::PropagateContainerPropertyChange( UObject* ModifiedObject, const FString& OriginalContainerContent, EPropertyArrayChangeType::Type ChangeType, int32 Index, TMap<UObject*, bool>* PropagationResult)
 {
 	UProperty* NodeProperty = GetProperty();
@@ -2815,10 +2804,6 @@ void FPropertyNode::PropagateContainerPropertyChange( UObject* ModifiedObject, c
 							Object->InstanceSubobjectTemplates();
 							break;
 						}
-						if (ElementToInitialize >= 0)
-						{
-							AdditionalInitializationUDS(ArrayProperty->Inner, ArrayHelper.GetRawPtr(ElementToInitialize));
-						}
 					}
 				}	// End Array
 
@@ -2850,11 +2835,6 @@ void FPropertyNode::PropagateContainerPropertyChange( UObject* ModifiedObject, c
 							check(false);	// Duplicate not supported on sets
 							break;
 						}
-
-						if (ElementToInitialize >= 0)
-						{
-							AdditionalInitializationUDS(SetProperty->ElementProp, SetHelper.GetElementPtr(ElementToInitialize));
-						}
 					}
 				}	// End Set
 				else if (MapProperty)
@@ -2885,14 +2865,6 @@ void FPropertyNode::PropagateContainerPropertyChange( UObject* ModifiedObject, c
 							check(false);	// Duplicate is not supported for maps
 							break;
 						}
-
-						if (ElementToInitialize >= 0)
-						{
-							uint8* PairPtr = MapHelper.GetPairPtr(ElementToInitialize);
-
-							AdditionalInitializationUDS(MapProperty->KeyProp, MapProperty->KeyProp->ContainerPtrToValuePtr<uint8>(PairPtr));
-							AdditionalInitializationUDS(MapProperty->ValueProp, MapProperty->ValueProp->ContainerPtrToValuePtr<uint8>(PairPtr));
-						}
 					}
 				}	// End Map
 			}
diff --git a/Engine/Source/Editor/PropertyEditor/Private/PropertyNode.h b/Engine/Source/Editor/PropertyEditor/Private/PropertyNode.h
--- a/Engine/Source/Editor/PropertyEditor/Private/PropertyNode.h
+++ b/Engine/Source/Editor/PropertyEditor/Private/PropertyNode.h
@@ -538,8 +538,6 @@ public:
 	void PropagateContainerPropertyChange(UObject* ModifiedObject, const FString& OriginalContainerContent,
 		EPropertyArrayChangeType::Type ChangeType, int32 Index, TMap<UObject*, bool>* PropagationResult = nullptr);
 
-	static void AdditionalInitializationUDS(UProperty* Property, uint8* RawPtr);
-
 	/** Broadcasts when a property value changes */
 	DECLARE_EVENT(FPropertyNode, FPropertyValueChangedEvent);
 	FPropertyValueChangedEvent& OnPropertyValueChanged() { return PropertyValueChangedEvent; }
diff --git a/Engine/Source/Editor/UnrealEd/Classes/UserDefinedStructure/UserDefinedStructEditorData.h b/Engine/Source/Editor/UnrealEd/Classes/UserDefinedStructure/UserDefinedStructEditorData.h
--- a/Engine/Source/Editor/UnrealEd/Classes/UserDefinedStructure/UserDefinedStructEditorData.h
+++ b/Engine/Source/Editor/UnrealEd/Classes/UserDefinedStructure/UserDefinedStructEditorData.h
@@ -73,7 +73,7 @@ struct FStructVariableDescription
 	uint8 bEnable3dWidget:1;
 
 	// CurrentDefaultValue stores the actual default value, after the DefaultValue was changed, and before the struct was recompiled
-	UPROPERTY()
+	UPROPERTY(NonTransactional)
 	FString CurrentDefaultValue;
 
 	UPROPERTY()
@@ -107,19 +107,6 @@ struct TStructOpsTypeTraits< FStructVariableDescription > : public TStructOpsTyp
 	};
 };
 
-class FStructOnScopeMember : public FStructOnScope
-{
-public:
-	FStructOnScopeMember() : FStructOnScope() {}
-
-	void Recreate(const UStruct* InScriptStruct)
-	{
-		Destroy();
-		ScriptStruct = InScriptStruct;
-		Initialize();
-	}
-};
-
 UCLASS()
 class UNREALED_API UUserDefinedStructEditorData : public UObject, public FEditorUndoClient
 {
@@ -137,16 +124,7 @@ public:
 	UPROPERTY()
 	FString ToolTip;
 
-	// optional super struct
-	UPROPERTY()
-	UScriptStruct* NativeBase;
-
-protected:
-	FStructOnScopeMember DefaultStructInstance;
-
 public:
-	static void AddReferencedObjects(UObject* InThis, FReferenceCollector& Collector);
-
 	// UObject interface.
 	virtual TSharedPtr<ITransactionObjectAnnotation> GetTransactionAnnotation() const override;
 	virtual void PostEditUndo() override;
diff --git a/Engine/Source/Editor/UnrealEd/Private/DataTableEditorUtils.cpp b/Engine/Source/Editor/UnrealEd/Private/DataTableEditorUtils.cpp
--- a/Engine/Source/Editor/UnrealEd/Private/DataTableEditorUtils.cpp
+++ b/Engine/Source/Editor/UnrealEd/Private/DataTableEditorUtils.cpp
@@ -60,11 +60,6 @@ uint8* FDataTableEditorUtils::AddRow(UDataTable* DataTable, FName RowName)
 	DataTable->RowStruct->InitializeStruct(RowData);
 	// And be sure to call DestroyScriptStruct later
 
-	if (auto UDStruct = Cast<const UUserDefinedStruct>(DataTable->RowStruct))
-	{
-		UDStruct->InitializeDefaultValue(RowData);
-	}
-
 	// Add to row map
 	DataTable->RowMap.Add(RowName, RowData);
 	BroadcastPostChange(DataTable, EDataTableChangeInfo::RowList);
@@ -189,7 +184,7 @@ bool FDataTableEditorUtils::DiffersFromDefault(UDataTable* DataTable, FName RowN
 
 		if (const UUserDefinedStruct* UDStruct = Cast<const UUserDefinedStruct>(DataTable->RowStruct))
 		{
-			bDiffers = UDStruct->DiffersFromDefaultValue(RowData);
+			return !UDStruct->CompareScriptStruct(RowData, UDStruct->GetDefaultInstance(), PPF_None);
 		}
 	}
 
diff --git a/Engine/Source/Editor/UnrealEd/Private/EditorServer.cpp b/Engine/Source/Editor/UnrealEd/Private/EditorServer.cpp
--- a/Engine/Source/Editor/UnrealEd/Private/EditorServer.cpp
+++ b/Engine/Source/Editor/UnrealEd/Private/EditorServer.cpp
@@ -5038,6 +5038,10 @@ bool UEditorEngine::Exec_Transaction(const TCHAR* Str, FOutputDevice& Ar)
 
 void UEditorEngine::BroadcastPostUndo(const FString& Context, UObject* PrimaryObject, bool bUndoSuccess )
 {
+	// This sanitization code can be removed once blueprint ::Conform(ImplementedEvents/ImplementedInterfaces) 
+	// functions have been fixed. For the time being it improves editor stability, though:
+	UEdGraphPin::SanitizePinsPostUndoRedo();
+
 	for (auto UndoIt = UndoClients.CreateIterator(); UndoIt; ++UndoIt)
 	{
 		FEditorUndoClient* Client = *UndoIt;
@@ -5050,6 +5054,10 @@ void UEditorEngine::BroadcastPostUndo(const FString& Context, UObject* PrimaryOb
 
 void UEditorEngine::BroadcastPostRedo(const FString& Context, UObject* PrimaryObject, bool bRedoSuccess )
 {
+	// This sanitization code can be removed once blueprint ::Conform(ImplementedEvents/ImplementedInterfaces) 
+	// functions have been fixed. For the time being it improves editor stability, though:
+	UEdGraphPin::SanitizePinsPostUndoRedo();
+
 	for (auto UndoIt = UndoClients.CreateIterator(); UndoIt; ++UndoIt)
 	{
 		FEditorUndoClient* Client = *UndoIt;
diff --git a/Engine/Source/Editor/UnrealEd/Private/Kismet2/BlueprintEditorUtils.cpp b/Engine/Source/Editor/UnrealEd/Private/Kismet2/BlueprintEditorUtils.cpp
--- a/Engine/Source/Editor/UnrealEd/Private/Kismet2/BlueprintEditorUtils.cpp
+++ b/Engine/Source/Editor/UnrealEd/Private/Kismet2/BlueprintEditorUtils.cpp
@@ -1407,6 +1407,7 @@ UClass* FBlueprintEditorUtils::RegenerateBlueprintClass(UBlueprint* Blueprint, U
 			if (Blueprint->GeneratedClass)
 			{
 				FBlueprintEditorUtils::RecreateClassMetaData(Blueprint, Blueprint->GeneratedClass, true);
+				Blueprint->GeneratedClass->SetUpRuntimeReplicationData();
 			}
 
 			// Flag data only blueprints as being up-to-date
@@ -8055,7 +8056,6 @@ bool FBlueprintEditorUtils::PropertyValueFromString_Direct(const UProperty* Prop
 			const int32 StructSize = Struct->GetStructureSize() * StructProperty->ArrayDim;
 			StructProperty->InitializeValue(DirectValue);
 			ensure(1 == StructProperty->ArrayDim);
-			bParseSucceeded = FStructureEditorUtils::Fill_MakeStructureDefaultValue(Cast<const UUserDefinedStruct>(Struct), DirectValue);
 
 			FStringOutputDevice ImportError;
 			const TCHAR* EndOfParsedBuff = StructProperty->ImportText(StrValue.IsEmpty() ? TEXT("()") : *StrValue, DirectValue, PPF_SerializedAsImportText, nullptr, &ImportError);
diff --git a/Engine/Source/Editor/UnrealEd/Private/Kismet2/ComponentEditorUtils.cpp b/Engine/Source/Editor/UnrealEd/Private/Kismet2/ComponentEditorUtils.cpp
--- a/Engine/Source/Editor/UnrealEd/Private/Kismet2/ComponentEditorUtils.cpp
+++ b/Engine/Source/Editor/UnrealEd/Private/Kismet2/ComponentEditorUtils.cpp
@@ -536,7 +536,7 @@ int32 FComponentEditorUtils::DeleteComponents(const TArray<UActorComponent*>& Co
 						ParentComponent->GetChildrenComponents(false, Siblings);
 						for (int32 i = 0; i < Siblings.Num() && ComponentToDelete != Siblings[i]; ++i)
 						{
-							if (!Siblings[i]->IsPendingKill())
+							if (Siblings[i] && !Siblings[i]->IsPendingKill())
 							{
 								OutComponentToSelect = Siblings[i];
 							}
diff --git a/Engine/Source/Editor/UnrealEd/Private/Kismet2/KismetDebugUtilities.cpp b/Engine/Source/Editor/UnrealEd/Private/Kismet2/KismetDebugUtilities.cpp
--- a/Engine/Source/Editor/UnrealEd/Private/Kismet2/KismetDebugUtilities.cpp
+++ b/Engine/Source/Editor/UnrealEd/Private/Kismet2/KismetDebugUtilities.cpp
@@ -28,6 +28,7 @@
 #include "K2Node_Tunnel.h"
 #include "K2Node_FunctionEntry.h"
 #include "K2Node_MacroInstance.h"
+#include "K2Node_Message.h"
 #include "Kismet2/KismetEditorUtilities.h"
 #include "Kismet2/BlueprintEditorUtils.h"
 #include "Logging/TokenizedMessage.h"
@@ -280,13 +281,10 @@ void FKismetDebugUtilities::OnScriptException(const UObject* ActiveObject, const
 
 			// Find the node that generated the code which we hit
 			UEdGraphNode* NodeStoppedAt = FindSourceNodeForCodeLocation(ActiveObject, StackFrame.Node, BreakpointOffset, /*bAllowImpreciseHit=*/ true);
-			if (NodeStoppedAt)
+			if (NodeStoppedAt && (Info.GetType() == EBlueprintExceptionType::Tracepoint || Info.GetType() == EBlueprintExceptionType::Breakpoint))
 			{
 				// Handle Node stepping and update the stack
-				if (Info.GetType() == EBlueprintExceptionType::Tracepoint || Info.GetType() == EBlueprintExceptionType::Breakpoint)
-				{
-					CheckBreakConditions(NodeStoppedAt, bShouldBreakExecution);
-				}
+				CheckBreakConditions(NodeStoppedAt, bShouldBreakExecution);
 			}
 
 			// Handle a breakpoint or single-step
@@ -527,6 +525,23 @@ void FKismetDebugUtilities::AttemptToBreakExecution(UBlueprint* BlueprintObj, co
 			Data.bIsSingleStepping = false;
 		}
 	}
+	else if(UEdGraphNode* PreviousNode = FKismetDebugUtilities::GetCurrentInstruction())
+	{
+		if (UK2Node_Message* MessageNode = Cast<UK2Node_Message>(PreviousNode))
+		{
+			//Looks like object not implement one of their interfaces
+			UE_LOG(LogBlueprintDebug, Warning, TEXT("Can't break execution on function '%s'. Possibly interface '%s' in class '%s' was not fully implemented."),
+				*(PreviousNode->GetDocumentationExcerptName()),					  //Function name
+				*(MessageNode->GetTargetFunction()->GetOuterUClass()->GetName()), //Interface name
+				*(ActiveObject->GetClass()->GetName()));						  //Current object class name
+		}
+		else
+		{
+			UE_LOG(LogBlueprintDebug, Warning, TEXT("Can't break execution on function '%s'. Possibly it was not implemented in class '%s'."),
+				*(PreviousNode->GetDocumentationExcerptName()),					  //Function name
+				*(ActiveObject->GetClass()->GetName()));						  //Current object class name
+		}
+	}
 	else
 	{
 		UE_LOG(LogBlueprintDebug, Warning, TEXT("Tried to break execution in an unknown spot at object %s:%04X"), *StackFrame.Node->GetFullName(), StackFrame.Code - StackFrame.Node->Script.GetData());
diff --git a/Engine/Source/Editor/UnrealEd/Private/Kismet2/StructureEditorUtils.cpp b/Engine/Source/Editor/UnrealEd/Private/Kismet2/StructureEditorUtils.cpp
--- a/Engine/Source/Editor/UnrealEd/Private/Kismet2/StructureEditorUtils.cpp
+++ b/Engine/Source/Editor/UnrealEd/Private/Kismet2/StructureEditorUtils.cpp
@@ -480,84 +480,34 @@ FString FStructureEditorUtils::GetVariableDisplayName(const UUserDefinedStruct*
 	return VarDesc ? VarDesc->FriendlyName : FString();
 }
 
-bool FStructureEditorUtils::UserDefinedStructEnabled()
-{
-	static FBoolConfigValueHelper UseUserDefinedStructure(TEXT("UserDefinedStructure"), TEXT("bUseUserDefinedStructure"));
-	return UseUserDefinedStructure;
-}
-
-void FStructureEditorUtils::RecreateDefaultInstanceInEditorData(UUserDefinedStruct* Struct)
+UProperty* FStructureEditorUtils::GetPropertyByDisplayName(const UUserDefinedStruct* Struct, FString DisplayName)
 {
-	UUserDefinedStructEditorData* StructEditorData = Struct ? CastChecked<UUserDefinedStructEditorData>(Struct->EditorData) : nullptr;
-	if (StructEditorData)
-	{
-		StructEditorData->RecreateDefaultInstance();
-	}
-}
-
-bool FStructureEditorUtils::Fill_MakeStructureDefaultValue(const UUserDefinedStruct* Struct, uint8* StructData)
-{
-	bool bResult = true;
-	if (Struct && StructData)
+	if (Struct)
 	{
-		UUserDefinedStructEditorData* StructEditorData = CastChecked<UUserDefinedStructEditorData>(Struct->EditorData);
-		const uint8* DefaultInstance = StructEditorData->GetDefaultInstance();
-		if (DefaultInstance)
-		{
-			Struct->CopyScriptStruct(StructData, DefaultInstance);
-		}
-		else
+		for (const FStructVariableDescription& VarDesc : GetVarDesc(Struct))
 		{
-			bResult = false;
+			if (VarDesc.FriendlyName == DisplayName)
+			{
+				return FindField<UProperty>(Struct, VarDesc.VarName);
+			}
 		}
 	}
-	
-	return bResult;
+	return nullptr;
 }
 
-bool FStructureEditorUtils::DiffersFromDefaultValue(const UUserDefinedStruct* Struct, uint8* StructData)
+bool FStructureEditorUtils::UserDefinedStructEnabled()
 {
-	bool bDiffers = false;
-	if (Struct && StructData)
-	{
-		UUserDefinedStructEditorData* StructEditorData = CastChecked<UUserDefinedStructEditorData>(Struct->EditorData);
-		const uint8* DefaultInstance = StructEditorData->GetDefaultInstance();
-		if (DefaultInstance)
-		{
-			const int32 PortFlags = PPF_None;
-			bDiffers = !Struct->CompareScriptStruct(StructData, DefaultInstance, PortFlags);
-		}
-	}
-	return bDiffers;
+	static FBoolConfigValueHelper UseUserDefinedStructure(TEXT("UserDefinedStructure"), TEXT("bUseUserDefinedStructure"));
+	return UseUserDefinedStructure;
 }
 
-bool FStructureEditorUtils::Fill_MakeStructureDefaultValue(const UProperty* Property, uint8* PropertyData)
+void FStructureEditorUtils::RecreateDefaultInstanceInEditorData(UUserDefinedStruct* Struct)
 {
-	bool bResult = true;
-
-	if (const UStructProperty* StructProperty = Cast<const UStructProperty>(Property))
-	{
-		if (const UUserDefinedStruct* InnerStruct = Cast<const UUserDefinedStruct>(StructProperty->Struct))
-		{
-			bResult &= Fill_MakeStructureDefaultValue(InnerStruct, PropertyData);
-		}
-	}
-	else if (const UArrayProperty* ArrayProp = Cast<const UArrayProperty>(Property))
+	UUserDefinedStructEditorData* StructEditorData = Struct ? CastChecked<UUserDefinedStructEditorData>(Struct->EditorData) : nullptr;
+	if (StructEditorData)
 	{
-		StructProperty = Cast<const UStructProperty>(ArrayProp->Inner);
-		const UUserDefinedStruct* InnerStruct = StructProperty ? Cast<const UUserDefinedStruct>(StructProperty->Struct) : NULL;
-		if(InnerStruct)
-		{
-			FScriptArrayHelper ArrayHelper(ArrayProp, PropertyData);
-			for (int32 Index = 0; Index < ArrayHelper.Num(); ++Index)
-			{
-				uint8* const ValuePtr = ArrayHelper.GetRawPtr(Index);
-				bResult &= Fill_MakeStructureDefaultValue(InnerStruct, ValuePtr);
-			}
-		}
+		StructEditorData->RecreateDefaultInstance();
 	}
-
-	return bResult;
 }
 
 void FStructureEditorUtils::CompileStructure(UUserDefinedStruct* Struct)
diff --git a/Engine/Source/Editor/UnrealEd/Private/UserDefinedStructEditorData.cpp b/Engine/Source/Editor/UnrealEd/Private/UserDefinedStructEditorData.cpp
--- a/Engine/Source/Editor/UnrealEd/Private/UserDefinedStructEditorData.cpp
+++ b/Engine/Source/Editor/UnrealEd/Private/UserDefinedStructEditorData.cpp
@@ -38,11 +38,6 @@ FEdGraphPinType FStructVariableDescription::ToPinType() const
 UUserDefinedStructEditorData::UUserDefinedStructEditorData(const FObjectInitializer& ObjectInitializer)
 	: Super(ObjectInitializer)
 {
-	UUserDefinedStruct* ScriptStruct = GetOwnerStruct();
-	if (ScriptStruct)
-	{
-		DefaultStructInstance.SetPackage(ScriptStruct->GetOutermost());
-	}
 }
 
 uint32 UUserDefinedStructEditorData::GenerateUniqueNameIdForMemberVariable()
@@ -130,31 +125,32 @@ void UUserDefinedStructEditorData::PostLoadSubobjects(FObjectInstancingGraph* Ou
 
 const uint8* UUserDefinedStructEditorData::GetDefaultInstance() const
 {
-	ensure(DefaultStructInstance.IsValid() && DefaultStructInstance.GetStruct() == GetOwnerStruct());
-	return DefaultStructInstance.GetStructMemory();
+	return GetOwnerStruct()->GetDefaultInstance();
 }
 
 void UUserDefinedStructEditorData::RecreateDefaultInstance(FString* OutLog)
 {
-	UStruct* ScriptStruct = GetOwnerStruct();
-	DefaultStructInstance.Recreate(ScriptStruct);
-	uint8* StructData = DefaultStructInstance.GetStructMemory();
-	ensure(DefaultStructInstance.IsValid() && DefaultStructInstance.GetStruct() == ScriptStruct);
-	if (DefaultStructInstance.IsValid() && StructData && ScriptStruct)
+	UUserDefinedStruct* ScriptStruct = GetOwnerStruct();
+	ScriptStruct->DefaultStructInstance.Recreate(ScriptStruct);
+	uint8* StructData = ScriptStruct->DefaultStructInstance.GetStructMemory();
+	ensure(ScriptStruct->DefaultStructInstance.IsValid() && ScriptStruct->DefaultStructInstance.GetStruct() == ScriptStruct);
+	if (ScriptStruct->DefaultStructInstance.IsValid() && StructData)
 	{
 		// When loading, the property's default value may end up being filled with a placeholder. 
 		// This tracker object allows the linker to track the actual object that is being filled in 
 		// so it can calculate an offset to the property and write in the placeholder value:
 		FScopedPlaceholderRawContainerTracker TrackDefaultObject(StructData);
 
-		DefaultStructInstance.SetPackage(ScriptStruct->GetOutermost());
+		ScriptStruct->DefaultStructInstance.SetPackage(ScriptStruct->GetOutermost());
 
 		for (TFieldIterator<UProperty> It(ScriptStruct); It; ++It)
 		{
 			UProperty* Property = *It;
 			if (Property)
 			{
-				FStructVariableDescription* VarDesc = VariablesDescriptions.FindByPredicate(FStructureEditorUtils::FFindByNameHelper<FStructVariableDescription>(Property->GetFName()));
+				FGuid VarGuid = FStructureEditorUtils::GetGuidFromPropertyName(Property->GetFName());
+
+				FStructVariableDescription* VarDesc = VariablesDescriptions.FindByPredicate(FStructureEditorUtils::FFindByGuidHelper<FStructVariableDescription>(VarGuid));
 				if (VarDesc && !VarDesc->CurrentDefaultValue.IsEmpty())
 				{
 					if (!FBlueprintEditorUtils::PropertyValueFromString(Property, VarDesc->CurrentDefaultValue, StructData))
@@ -176,23 +172,9 @@ void UUserDefinedStructEditorData::RecreateDefaultInstance(FString* OutLog)
 
 void UUserDefinedStructEditorData::CleanDefaultInstance()
 {
-	ensure(!DefaultStructInstance.IsValid() || DefaultStructInstance.GetStruct() == GetOwnerStruct());
-	DefaultStructInstance.Destroy();
-}
-
-void UUserDefinedStructEditorData::AddReferencedObjects(UObject* InThis, FReferenceCollector& Collector)
-{
-	UUserDefinedStructEditorData* This = CastChecked<UUserDefinedStructEditorData>(InThis);
-
-	UStruct* ScriptStruct = This->GetOwnerStruct();
-	ensure(!This->DefaultStructInstance.IsValid() || This->DefaultStructInstance.GetStruct() == ScriptStruct);
-	uint8* StructData = This->DefaultStructInstance.GetStructMemory();
-	if (StructData)
-	{
-		ScriptStruct->SerializeBin(Collector.GetVerySlowReferenceCollectorArchive(), StructData);
-	}
-
-	Super::AddReferencedObjects(This, Collector);
+	UUserDefinedStruct* ScriptStruct = GetOwnerStruct();
+	ensure(!ScriptStruct->DefaultStructInstance.IsValid() || ScriptStruct->DefaultStructInstance.GetStruct() == GetOwnerStruct());
+	ScriptStruct->DefaultStructInstance.Destroy();
 }
 
 #undef LOCTEXT_NAMESPACE
diff --git a/Engine/Source/Editor/UnrealEd/Public/Kismet2/StructureEditorUtils.h b/Engine/Source/Editor/UnrealEd/Public/Kismet2/StructureEditorUtils.h
--- a/Engine/Source/Editor/UnrealEd/Public/Kismet2/StructureEditorUtils.h
+++ b/Engine/Source/Editor/UnrealEd/Public/Kismet2/StructureEditorUtils.h
@@ -91,6 +91,8 @@ public:
 
 	static FString GetVariableDisplayName(const UUserDefinedStruct* Struct, FGuid VarGuid);
 
+	static UProperty* GetPropertyByDisplayName(const UUserDefinedStruct* Struct, FString DisplayName);
+
 	static FString GetVariableTooltip(const UUserDefinedStruct* Struct, FGuid VarGuid);
 	
 	static bool ChangeVariableTooltip(UUserDefinedStruct* Struct, FGuid VarGuid, const FString& InTooltip);
@@ -145,18 +147,8 @@ public:
 	static void RemoveInvalidStructureMemberVariableFromBlueprint(UBlueprint* Blueprint);
 
 	//DEFAULT VALUE
-	/*
-	 Default values for member variables in User Defined Structure are stored in meta data "MakeStructureDefaultValue"
-	 The following functions are used to fill an instance of user defined struct with those default values.
-	 */
-	static bool Fill_MakeStructureDefaultValue(const UUserDefinedStruct* Struct, uint8* StructData);
-
-	static bool Fill_MakeStructureDefaultValue(const UProperty* Property, uint8* PropertyData);
-
 	static void RecreateDefaultInstanceInEditorData(UUserDefinedStruct* Struct);
 
-	static bool DiffersFromDefaultValue(const UUserDefinedStruct* Struct, uint8* StructData);
-
 	//VALIDATION
 	static bool CanHaveAMemberVariableOfType(const UUserDefinedStruct* Struct, const FEdGraphPinType& VarType, FString* OutMsg = NULL);
 
diff --git a/Engine/Source/Runtime/Core/Public/UObject/FrameworkObjectVersion.h b/Engine/Source/Runtime/Core/Public/UObject/FrameworkObjectVersion.h
--- a/Engine/Source/Runtime/Core/Public/UObject/FrameworkObjectVersion.h
+++ b/Engine/Source/Runtime/Core/Public/UObject/FrameworkObjectVersion.h
@@ -110,6 +110,9 @@ struct CORE_API FFrameworkObjectVersion
 		// FMaterialInput and FEdGraphPin store their name as FName instead of FString
 		PinsStoreFName,
 
+		// User defined structs store their default instance, which is used for initializing instances
+		UserDefinedStructsStoreDefaultInstance,
+
 		// -----<new versions can be added above this line>-------------------------------------------------
 		VersionPlusOne,
 		LatestVersion = VersionPlusOne - 1
diff --git a/Engine/Source/Runtime/CoreUObject/Private/Blueprint/BlueprintSupport.cpp b/Engine/Source/Runtime/CoreUObject/Private/Blueprint/BlueprintSupport.cpp
--- a/Engine/Source/Runtime/CoreUObject/Private/Blueprint/BlueprintSupport.cpp
+++ b/Engine/Source/Runtime/CoreUObject/Private/Blueprint/BlueprintSupport.cpp
@@ -132,6 +132,60 @@ bool FBlueprintSupport::IsDeferredDependencyPlaceholder(UObject* LoadedObj)
 		LoadedObj->IsA<ULinkerPlaceholderExportObject>() );
 }
 
+void FBlueprintSupport::RegisterDeferredDependenciesInStruct(const UStruct* Struct, void* StructData)
+{
+#if USE_CIRCULAR_DEPENDENCY_LOAD_DEFERRING
+	if (GEventDrivenLoaderEnabled)
+	{
+		return;
+	}
+
+	for (TPropertyValueIterator<const UObjectProperty> It(Struct, StructData); It; ++It)
+	{
+		const UObjectProperty* Property = It.Key();
+		void* PropertyValue = (void*)It.Value();
+		UObject* ObjectValue = *((UObject**)PropertyValue);
+		
+		ULinkerPlaceholderExportObject* PlaceholderVal = Cast<ULinkerPlaceholderExportObject>(ObjectValue);
+		ULinkerPlaceholderClass* PlaceholderClass = Cast<ULinkerPlaceholderClass>(ObjectValue);
+
+		if (PlaceholderVal == nullptr && PlaceholderClass == nullptr)
+		{
+			continue;
+		}
+
+		// Create a stack of property trackers to deal with any outer Struct Properties
+		TArray<const UProperty*> PropertyChain;
+		It.GetPropertyChain(PropertyChain);
+		TIndirectArray<FScopedPlaceholderPropertyTracker> PlaceholderStack;
+
+		// Iterate property chain in reverse order as we need to start with parent
+		for (int32 PropertyIndex = PropertyChain.Num() - 1; PropertyIndex >= 0; PropertyIndex--)
+		{
+			if (const UStructProperty* StructProperty = Cast<UStructProperty>(PropertyChain[PropertyIndex]))
+			{
+				PlaceholderStack.Add(new FScopedPlaceholderPropertyTracker(StructProperty));
+			}
+		}
+		
+		if (PlaceholderVal)
+		{
+			PlaceholderVal->AddReferencingPropertyValue(Property, PropertyValue);
+		}
+		else 
+		{
+			PlaceholderClass->AddReferencingPropertyValue(Property, PropertyValue);
+		}
+
+		// Specifically destroy entries in reverse order they were added, to simulate unrolling a code stack
+		for (int32 StackIndex = PlaceholderStack.Num() - 1; StackIndex >= 0; StackIndex--)
+		{
+			PlaceholderStack.RemoveAt(StackIndex);
+		}
+	}
+#endif // USE_CIRCULAR_DEPENDENCY_LOAD_DEFERRING
+}
+
 bool FBlueprintSupport::IsInBlueprintPackage(UObject* LoadedObj)
 {
 	UPackage* Pkg = LoadedObj->GetOutermost();
@@ -2584,6 +2638,17 @@ void FConvertedBlueprintsDependencies::FillUsedAssetsInDynamicClass(UDynamicClas
 	}
 }
 
+UObject* FConvertedBlueprintsDependencies::LoadObjectForStructConstructor(UScriptStruct* ScriptStruct, const TCHAR* ObjectPath)
+{
+	if (GEventDrivenLoaderEnabled && EVENT_DRIVEN_ASYNC_LOAD_ACTIVE_AT_RUNTIME)
+	{
+		// Find Object should work here as the blueprints have scheduled it for load
+		return FindObject<UObject>(nullptr, ObjectPath);
+	}
+
+	return LoadObject<UObject>(nullptr, ObjectPath);
+}
+
 bool FBlueprintDependencyData::ContainsDependencyData(TArray<FBlueprintDependencyData>& Assets, int16 ObjectRefIndex)
 {
 	return nullptr != Assets.FindByPredicate([=](const FBlueprintDependencyData& Data) -> bool
diff --git a/Engine/Source/Runtime/CoreUObject/Private/Serialization/AsyncLoading.cpp b/Engine/Source/Runtime/CoreUObject/Private/Serialization/AsyncLoading.cpp
--- a/Engine/Source/Runtime/CoreUObject/Private/Serialization/AsyncLoading.cpp
+++ b/Engine/Source/Runtime/CoreUObject/Private/Serialization/AsyncLoading.cpp
@@ -4612,21 +4612,7 @@ EAsyncPackageState::Type FAsyncLoadingThread::ProcessLoadedPackages(bool bUseTim
 
 #if WITH_EDITOR
 				// In the editor we need to find any assets and add them to list for later callback
-				UPackage* LoadedPackage = Package->GetLoadedPackage();
-
-				if (LoadedPackage)
-				{
-					TArray<UObject*> TopLevelObjects;
-					GetObjectsWithOuter(LoadedPackage, TopLevelObjects, false);
-
-					for (UObject* TopLevelObject : TopLevelObjects)
-					{
-						if (TopLevelObject->IsAsset())
-						{
-							LoadedAssets.Add(TopLevelObject);
-						}
-					}
-				}
+				Package->GetLoadedAssets(LoadedAssets);
 #endif
 				// We don't need the package anymore
 				PackagesToDelete.AddUnique(Package);
@@ -5361,6 +5347,19 @@ void FAsyncPackage::FlushObjectLinkerCache()
 	}
 }
 
+#if WITH_EDITOR 
+void FAsyncPackage::GetLoadedAssets(TArray<FWeakObjectPtr>& AssetList)
+{
+	for (UObject* Obj : PackageObjLoaded)
+	{
+		if (Obj && !Obj->IsPendingKill() && Obj->IsAsset())
+		{
+			AssetList.Add(Obj);
+		}
+	}
+}
+#endif
+
 /**
  * Gives up time slice if time limit is enabled.
  *
diff --git a/Engine/Source/Runtime/CoreUObject/Private/UObject/Class.cpp b/Engine/Source/Runtime/CoreUObject/Private/UObject/Class.cpp
--- a/Engine/Source/Runtime/CoreUObject/Private/UObject/Class.cpp
+++ b/Engine/Source/Runtime/CoreUObject/Private/UObject/Class.cpp
@@ -3011,6 +3011,11 @@ void UClass::PostLoad()
 	{
 		GetSuperClass()->ConditionalPostLoad();
 	}
+
+	if (!HasAnyClassFlags(CLASS_Native))
+	{
+		SetUpRuntimeReplicationData();
+	}
 }
 
 FString UClass::GetDesc()
@@ -3042,6 +3047,16 @@ void UClass::Link(FArchive& Ar, bool bRelinkExistingProperties)
 	check(!bRelinkExistingProperties || !(ClassFlags & CLASS_Intrinsic));
 	Super::Link(Ar, bRelinkExistingProperties);
 
+	// For non-native classes, this will happen in PostLoad when all its functions
+	// are guaranteed to be loaded. Native classes have to do this now.
+	if (HasAnyClassFlags(CLASS_Native))
+	{
+		SetUpRuntimeReplicationData();
+	}
+}
+
+void UClass::SetUpRuntimeReplicationData()
+{
 	if (PropertyLink != NULL)
 	{
 		NetFields.Empty();
@@ -3074,7 +3089,11 @@ void UClass::Link(FArchive& Ar, bool bRelinkExistingProperties)
 			}
 			else if( (F=dynamic_cast<UFunction*>(*It))!=NULL )
 			{
-				check(!F->GetSuperFunction() || (F->GetSuperFunction()->FunctionFlags&FUNC_NetFuncFlags) == (F->FunctionFlags&FUNC_NetFuncFlags));
+				// When loading reflection data (e.g. from blueprints), we may have references to placeholder functions, or reflection data 
+				// in children may be out of date. In that case we cannot enforce this check, but that is ok because reflection data will
+				// be regenerated by compile on load anyway:
+				const bool bCanCheck = (!GIsEditor && !IsRunningCommandlet()) || !F->HasAnyFlags(RF_WasLoaded);
+				check(!bCanCheck || (!F->GetSuperFunction() || (F->GetSuperFunction()->FunctionFlags&FUNC_NetFuncFlags) == (F->FunctionFlags&FUNC_NetFuncFlags)));
 				if( (F->FunctionFlags&FUNC_Net) && !F->GetSuperFunction() )
 					NetFields.Add( *It );
 			}
diff --git a/Engine/Source/Runtime/CoreUObject/Private/UObject/CoreNet.cpp b/Engine/Source/Runtime/CoreUObject/Private/UObject/CoreNet.cpp
--- a/Engine/Source/Runtime/CoreUObject/Private/UObject/CoreNet.cpp
+++ b/Engine/Source/Runtime/CoreUObject/Private/UObject/CoreNet.cpp
@@ -167,6 +167,8 @@ const FClassNetCache* FClassNetCacheMgr::GetClassNetCache( const UClass* Class )
 
 	if ( !Result )
 	{
+		ensureMsgf(!Class->HasAnyFlags(RF_NeedLoad | RF_NeedPostLoad), TEXT("FClassNetCacheMgr::GetClassNetCache: %s has flag RF_NeedPostLoad. NetFields and ClassReps will be incorrect!"), *GetFullNameSafe(Class));
+
 		Result					= ClassFieldIndices.Add( Class, new FClassNetCache( Class ) );
 		Result->Super			= NULL;
 		Result->FieldsBase		= 0;
diff --git a/Engine/Source/Runtime/CoreUObject/Private/UObject/LinkerPlaceholderBase.cpp b/Engine/Source/Runtime/CoreUObject/Private/UObject/LinkerPlaceholderBase.cpp
--- a/Engine/Source/Runtime/CoreUObject/Private/UObject/LinkerPlaceholderBase.cpp
+++ b/Engine/Source/Runtime/CoreUObject/Private/UObject/LinkerPlaceholderBase.cpp
@@ -234,8 +234,6 @@ FScopedPlaceholderContainerTracker::~FScopedPlaceholderContainerTracker()
 #endif // USE_DEFERRED_DEPENDENCY_CHECK_VERIFICATION_TESTS
 }
 
-#if WITH_EDITOR
-
 FScopedPlaceholderRawContainerTracker::FScopedPlaceholderRawContainerTracker(void* InData)
 	:Data(InData)
 {
@@ -250,8 +248,6 @@ FScopedPlaceholderRawContainerTracker::~FScopedPlaceholderRawContainerTracker()
 #endif // USE_DEFERRED_DEPENDENCY_CHECK_VERIFICATION_TESTS
 }
 
-#endif
-
 //------------------------------------------------------------------------------
 FScopedPlaceholderPropertyTracker::FScopedPlaceholderPropertyTracker(const UStructProperty* InIntermediateProperty)
 	: IntermediateProperty(nullptr) // leave null, as a sentinel value (implying that PerspectiveReferencerStack was empty)
diff --git a/Engine/Source/Runtime/CoreUObject/Public/Blueprint/BlueprintSupport.h b/Engine/Source/Runtime/CoreUObject/Public/Blueprint/BlueprintSupport.h
--- a/Engine/Source/Runtime/CoreUObject/Public/Blueprint/BlueprintSupport.h
+++ b/Engine/Source/Runtime/CoreUObject/Public/Blueprint/BlueprintSupport.h
@@ -58,6 +58,9 @@ struct FBlueprintSupport
 	/** Tells if the specified object is one of the many flavors of FLinkerPlaceholderBase that we have. */
 	COREUOBJECT_API static bool IsDeferredDependencyPlaceholder(UObject* LoadedObj);
 
+	/** Registers any object properties in this struct with the deferred dependency system */
+	COREUOBJECT_API static void RegisterDeferredDependenciesInStruct(const UStruct* Struct, void* StructData);
+
 	/** Not a particularly fast function. Mostly intended for validation in debug builds. */
 	static bool IsInBlueprintPackage(UObject* LoadedObj);
 
@@ -76,6 +79,21 @@ struct FBlueprintSupport
 #endif
 };
 
+/**
+ * When dealing with user defined structs we don't always have a UObject container
+ * this registers raw addresses for tracking. This is somewhat less safe, make
+ * sure to not register addresses that may change
+ */
+struct COREUOBJECT_API FScopedPlaceholderRawContainerTracker
+{
+public:
+	FScopedPlaceholderRawContainerTracker(void* InData);
+	~FScopedPlaceholderRawContainerTracker();
+
+private:
+	void* Data;
+};
+
 #if WITH_EDITOR
 /**
  * This is a helper struct that allows us to gather all previously unloaded class dependencies of a UClass
@@ -125,16 +143,6 @@ enum class EReplacementResult
 	ReplaceCompletely
 };
 
-struct COREUOBJECT_API FScopedPlaceholderRawContainerTracker
-{
-public:
-	FScopedPlaceholderRawContainerTracker(void* InData);
-	~FScopedPlaceholderRawContainerTracker();
-
-private:
-	void* Data;
-};
-
 /**
  * Interface needed by CoreUObject to the BlueprintNativeCodeGen logic. Used by cooker to convert assets 
  * to native code.
@@ -356,4 +364,5 @@ public:
 	void GetAssets(FName PackageName, TArray<FBlueprintDependencyData>& OutDependencies) const;
 
 	static void FillUsedAssetsInDynamicClass(UDynamicClass* DynamicClass, GetDependenciesNamesFunc GetUsedAssets);
+	static UObject* LoadObjectForStructConstructor(UScriptStruct* ScriptStruct, const TCHAR* ObjectPath);
 };
diff --git a/Engine/Source/Runtime/CoreUObject/Public/Serialization/AsyncLoading.h b/Engine/Source/Runtime/CoreUObject/Public/Serialization/AsyncLoading.h
--- a/Engine/Source/Runtime/CoreUObject/Public/Serialization/AsyncLoading.h
+++ b/Engine/Source/Runtime/CoreUObject/Public/Serialization/AsyncLoading.h
@@ -487,6 +487,11 @@ struct FAsyncPackage : FGCObject
 	/** Returns the UPackage wrapped by this, if it is valid */
 	UPackage* GetLoadedPackage();
 
+#if WITH_EDITOR
+	/** Gets all assets loaded by this async package, used in the editor */
+	void GetLoadedAssets(TArray<FWeakObjectPtr>& AssetList);
+#endif
+
 private:	
 
 	struct FCompletionCallback
diff --git a/Engine/Source/Runtime/CoreUObject/Public/UObject/Class.h b/Engine/Source/Runtime/CoreUObject/Public/UObject/Class.h
--- a/Engine/Source/Runtime/CoreUObject/Public/UObject/Class.h
+++ b/Engine/Source/Runtime/CoreUObject/Public/UObject/Class.h
@@ -2578,6 +2578,14 @@ public:
 	 */
 	virtual void GetRequiredPreloadDependencies(TArray<UObject*>& DependenciesOut) {}
 
+	/**
+	 * Initializes the ClassReps and NetFields arrays used by replication.
+	 * For classes that are loaded, this needs to happen in PostLoad to
+	 * ensure all replicated UFunctions have been serialized. For native classes,
+	 * this should happen in Link. Also needs to happen after blueprint compiliation.
+	 */
+	void SetUpRuntimeReplicationData();
+
 private:
 	#if UCLASS_FAST_ISA_IMPL == UCLASS_ISA_INDEXTREE
 		// For UObjectBaseUtility
diff --git a/Engine/Source/Runtime/CoreUObject/Public/UObject/StructOnScope.h b/Engine/Source/Runtime/CoreUObject/Public/UObject/StructOnScope.h
--- a/Engine/Source/Runtime/CoreUObject/Public/UObject/StructOnScope.h
+++ b/Engine/Source/Runtime/CoreUObject/Public/UObject/StructOnScope.h
@@ -14,6 +14,8 @@ protected:
 	TWeakObjectPtr<const UStruct> ScriptStruct;
 	uint8* SampleStructMemory;
 	TWeakObjectPtr<UPackage> Package;
+	/** Whether the struct memory is owned by this instance. */
+	bool OwnsMemory;
 
 	FStructOnScope()
 		: SampleStructMemory(nullptr)
@@ -21,7 +23,7 @@ protected:
 	{
 	}
 
-	void Initialize()
+	virtual void Initialize()
 	{
 		if (ScriptStruct.IsValid())
 		{
@@ -114,9 +116,4 @@ private:
 
 	FStructOnScope(const FStructOnScope&);
 	FStructOnScope& operator=(const FStructOnScope&);
-
-private:
-
-	/** Whether the struct memory is owned by this instance. */
-	bool OwnsMemory;
 };
diff --git a/Engine/Source/Runtime/Engine/Classes/EdGraph/EdGraphPin.h b/Engine/Source/Runtime/Engine/Classes/EdGraph/EdGraphPin.h
--- a/Engine/Source/Runtime/Engine/Classes/EdGraph/EdGraphPin.h
+++ b/Engine/Source/Runtime/Engine/Classes/EdGraph/EdGraphPin.h
@@ -581,6 +581,7 @@ public:
 
 	ENGINE_API static bool AreOrphanPinsEnabled();
 
+	ENGINE_API static void SanitizePinsPostUndoRedo();
 private:
 	/** Private Constructor. Create pins using CreatePin since all pin instances are managed by TSharedPtr. */
 	UEdGraphPin(UEdGraphNode* InOwningNode, const FGuid& PinGuid);
diff --git a/Engine/Source/Runtime/Engine/Classes/Engine/BlueprintGeneratedClass.h b/Engine/Source/Runtime/Engine/Classes/Engine/BlueprintGeneratedClass.h
--- a/Engine/Source/Runtime/Engine/Classes/Engine/BlueprintGeneratedClass.h
+++ b/Engine/Source/Runtime/Engine/Classes/Engine/BlueprintGeneratedClass.h
@@ -67,6 +67,9 @@ public:
 	// Map from source node (impure) to pure node script code range
 	TMap< TWeakObjectPtr<UEdGraphNode>, FInt32Range > PureNodeScriptCodeRangeMap;
 
+	// Reverse map from code offset to outer expansion (tunnel instance) source nodes
+	TMap< int32, TArray<TWeakObjectPtr<UEdGraphNode> > > LineNumberToTunnelInstanceSourceNodesMap;
+
 public:
 	FDebuggingInfoForSingleFunction()
 	{
@@ -326,6 +329,16 @@ public:
 		return Result;
 	}
 
+	const TArray<TWeakObjectPtr<UEdGraphNode> >* FindExpansionSourceNodesFromCodeLocation(UFunction* Function, int32 CodeOffset) const
+	{
+		if (const FDebuggingInfoForSingleFunction* pFuncInfo = PerFunctionLineNumbers.Find(Function))
+		{
+			return pFuncInfo->LineNumberToTunnelInstanceSourceNodesMap.Find(CodeOffset);
+		}
+
+		return nullptr;
+	}
+
 	// Finds the breakpoint injection site(s) in bytecode if any were associated with the given node
 	void FindBreakpointInjectionSites(UEdGraphNode* Node, TArray<uint8*>& InstallSites) const
 	{
@@ -373,7 +386,7 @@ public:
 	}
 
 	// Adds a debug record for a source node and destination in the bytecode of a specified function
-	void RegisterNodeToCodeAssociation(UEdGraphNode* SourceNode, UFunction* InFunction, int32 CodeOffset, bool bBreakpointSite)
+	void RegisterNodeToCodeAssociation(UEdGraphNode* SourceNode, const TArray<TWeakObjectPtr<UEdGraphNode> >& ExpansionSourceNodes, UFunction* InFunction, int32 CodeOffset, bool bBreakpointSite)
 	{
 		//@TODO: Nasty expansion behavior during compile time
 		if (bBreakpointSite)
@@ -384,6 +397,11 @@ public:
 
 		FDebuggingInfoForSingleFunction& PerFuncInfo = PerFunctionLineNumbers.FindOrAdd(InFunction);
 		PerFuncInfo.LineNumberToSourceNodeMap.Add(CodeOffset, SourceNode);
+
+		if (ExpansionSourceNodes.Num() > 0)
+		{
+			PerFuncInfo.LineNumberToTunnelInstanceSourceNodesMap.Add(CodeOffset, ExpansionSourceNodes);
+		}
 	}
 
 	void RegisterPureNodeScriptCodeRange(UEdGraphNode* SourceNode, UFunction* InFunction, FInt32Range InPureNodeScriptCodeRange)
diff --git a/Engine/Source/Runtime/Engine/Classes/Engine/UserDefinedStruct.h b/Engine/Source/Runtime/Engine/Classes/Engine/UserDefinedStruct.h
--- a/Engine/Source/Runtime/Engine/Classes/Engine/UserDefinedStruct.h
+++ b/Engine/Source/Runtime/Engine/Classes/Engine/UserDefinedStruct.h
@@ -7,8 +7,12 @@
 #include "Misc/Guid.h"
 #include "UObject/Class.h"
 #include "UObject/WeakObjectPtr.h"
+#include "UObject/StructOnScope.h"
 #include "UserDefinedStruct.generated.h"
 
+class UUserDefinedStruct;
+class UUserDefinedStructEditorData;
+
 UENUM()
 enum EUserDefinedStructureStatus
 {
@@ -24,6 +28,25 @@ enum EUserDefinedStructureStatus
 	UDSS_MAX,
 };
 
+/** Wrapper for StructOnScope that tells it to ignore default values */
+class ENGINE_API FUserStructOnScopeIgnoreDefaults : public FStructOnScope
+{
+public:
+	/** Constructor with no script struct, call Recreate later */
+	FUserStructOnScopeIgnoreDefaults() : FStructOnScope() {}
+
+	/** Constructor that initializes for you */
+	FUserStructOnScopeIgnoreDefaults(const UUserDefinedStruct* InUserStruct);
+
+	/** Initialize from existing data, will free when scope closes */
+	FUserStructOnScopeIgnoreDefaults(const UUserDefinedStruct* InUserStruct, uint8* InData);
+
+	/** Destroys and creates new struct */
+	void Recreate(const UUserDefinedStruct* InUserStruct);
+
+	virtual void Initialize() override;
+};
+
 UCLASS()
 class ENGINE_API UUserDefinedStruct : public UScriptStruct
 {
@@ -31,10 +54,7 @@ class ENGINE_API UUserDefinedStruct : public UScriptStruct
 
 public:
 #if WITH_EDITORONLY_DATA
-	UPROPERTY()
-	TEnumAsByte<enum EUserDefinedStructureStatus> Status;
-
-	// the original struct, when current struct isn't a temporary duplicate, the field should be null
+	/** The original struct, when current struct isn't a temporary duplicate, the field should be null */
 	UPROPERTY(Transient)
 	TWeakObjectPtr<UUserDefinedStruct> PrimaryStruct;
 
@@ -45,41 +65,63 @@ public:
 	UObject* EditorData;
 #endif // WITH_EDITORONLY_DATA
 
+	/** Status of this struct, outside of the editor it is assumed to always be UpToDate */
+	UPROPERTY()
+	TEnumAsByte<enum EUserDefinedStructureStatus> Status;
+
+	/** Uniquely identifies this specific user struct */
 	UPROPERTY()
 	FGuid Guid;
 
+protected:
+	/** Default instance of this struct with default values filled in, used to initialize structure */
+	FUserStructOnScopeIgnoreDefaults DefaultStructInstance;
+
+	/** Bool to indicate we want to initialize a version of this struct without defaults, this is set while allocating the DefaultStructInstance itself */
+	bool bIgnoreStructDefaults;
+
+public:
 #if WITH_EDITOR
 	// UObject interface.
-	virtual void Serialize(FArchive& Ar) override;
 	virtual void PostDuplicate(bool bDuplicateForPIE) override;
 	virtual void GetAssetRegistryTags(TArray<FAssetRegistryTag>& OutTags) const override;
 	virtual void PostLoad() override;
 	// End of UObject interface.
 
 	// UScriptStruct interface.
-	virtual void InitializeDefaultValue(uint8* StructData) const override;
+	virtual UProperty* CustomFindProperty(const FName Name) const override;
 	// End of UScriptStruct interface.
 
-	bool DiffersFromDefaultValue(uint8* StructData) const;
-
+	/** Creates a new guid if needed */
 	void ValidateGuid();
+
+	friend UUserDefinedStructEditorData;
 #endif	// WITH_EDITOR
 
+	// UObject interface.
+	virtual void Serialize(FArchive& Ar) override;
 	virtual void SerializeTaggedProperties(FArchive& Ar, uint8* Data, UStruct* DefaultsStruct, uint8* Defaults, const UObject* BreakRecursionIfFullyLoad=NULL) const override;
 	virtual FString PropertyNameToDisplayName(FName InName) const override;
+	// End of UObject interface.
 
 	// UScriptStruct interface.
+	virtual void InitializeStruct(void* Dest, int32 ArrayDim = 1) const override;
 	virtual uint32 GetStructTypeHash(const void* Src) const override;
 	virtual void RecursivelyPreload() override;
 	virtual FGuid GetCustomGuid() const override;
 	virtual FString GetStructCPPName() const override;
 	// End of  UScriptStruct interface.
 
+	/** Returns the raw memory of the default instance */
+	const uint8* GetDefaultInstance() const;
+
+	/** Specifically initialize this struct without using the default instance data */
+	void InitializeStructIgnoreDefaults(void* Dest, int32 ArrayDim = 1) const;
+
+	/** Computes hash */
 	static uint32 GetUserDefinedStructTypeHash(const void* Src, const UScriptStruct* Type);
 
-#if WITH_EDITOR
-public:
-	// UStruct
-	virtual UProperty* CustomFindProperty(const FName Name) const override;
-#endif	// WITH_EDITOR
+	/** returns references from default instance */
+	static void AddReferencedObjects(UObject* InThis, FReferenceCollector& Collector);
+
 };
diff --git a/Engine/Source/Runtime/Engine/Classes/GameFramework/Actor.h b/Engine/Source/Runtime/Engine/Classes/GameFramework/Actor.h
--- a/Engine/Source/Runtime/Engine/Classes/GameFramework/Actor.h
+++ b/Engine/Source/Runtime/Engine/Classes/GameFramework/Actor.h
@@ -847,9 +847,9 @@ public:
 	 * @param NewRotation		The new rotation for the Actor.
 	 * @param bSweep			Whether we sweep to the destination location, triggering overlaps along the way and stopping short of the target if blocked by something.
 	 *							Only the root component is swept and checked for blocking collision, child components move without sweeping. If collision is off, this has no effect.
-	 * @param bTeleport			Whether we teleport the physics state (if physics collision is enabled for this object).
-	 *							If true, physics velocity for this object is unchanged (so ragdoll parts are not affected by change in location).
-	 *							If false, physics velocity is updated based on the change in position (affecting ragdoll parts).
+	 * @param Teleport			How we teleport the physics state (if physics collision is enabled for this object).
+	 *							If equal to ETeleportType::TeleportPhysics, physics velocity for this object is unchanged (so ragdoll parts are not affected by change in location).
+	 *							If equal to ETeleportType::None, physics velocity is updated based on the change in position (affecting ragdoll parts).
 	 *							If CCD is on and not teleporting, this will affect objects along the entire swept volume.
 	 * @param OutSweepHitResult	The hit result from the move if swept.
 	 * @return	Whether the rotation was successfully set.
@@ -3167,22 +3167,43 @@ FORCEINLINE_DEBUGGABLE bool AActor::IsNetMode(ENetMode Mode) const
 
 #define HIDE_ACTOR_TRANSFORM_FUNCTIONS() private: \
 	FTransform GetTransform() const { return Super::GetTransform(); } \
+	FTransform GetActorTransform() const { return Super::GetActorTransform(); } \
 	FVector GetActorLocation() const { return Super::GetActorLocation(); } \
 	FRotator GetActorRotation() const { return Super::GetActorRotation(); } \
 	FQuat GetActorQuat() const { return Super::GetActorQuat(); } \
 	FVector GetActorScale() const { return Super::GetActorScale(); } \
-	bool SetActorLocation(const FVector& NewLocation, bool bSweep=false, FHitResult* OutSweepHitResult=nullptr) { return Super::SetActorLocation(NewLocation, bSweep, OutSweepHitResult); } \
-	bool SetActorRotation(FRotator NewRotation) { return Super::SetActorRotation(NewRotation); } \
-	bool SetActorRotation(const FQuat& NewRotation) { return Super::SetActorRotation(NewRotation); } \
-	bool SetActorLocationAndRotation(FVector NewLocation, FRotator NewRotation, bool bSweep=false, FHitResult* OutSweepHitResult=nullptr) { return Super::SetActorLocationAndRotation(NewLocation, NewRotation, bSweep, OutSweepHitResult); } \
-	bool SetActorLocationAndRotation(FVector NewLocation, const FQuat& NewRotation, bool bSweep=false, FHitResult* OutSweepHitResult=nullptr) { return Super::SetActorLocationAndRotation(NewLocation, NewRotation, bSweep, OutSweepHitResult); } \
+	bool SetActorTransform(const FTransform& NewTransform, bool bSweep=false, FHitResult* OutSweepHitResult=nullptr, ETeleportType Teleport = ETeleportType::None) { return Super::SetActorTransform(NewTransform, bSweep, OutSweepHitResult, Teleport); } \
+	bool SetActorLocation(const FVector& NewLocation, bool bSweep=false, FHitResult* OutSweepHitResult=nullptr, ETeleportType Teleport = ETeleportType::None) { return Super::SetActorLocation(NewLocation, bSweep, OutSweepHitResult, Teleport); } \
+	bool SetActorRotation(FRotator NewRotation, ETeleportType Teleport = ETeleportType::None) { return Super::SetActorRotation(NewRotation, Teleport); } \
+	bool SetActorRotation(const FQuat& NewRotation, ETeleportType Teleport = ETeleportType::None) { return Super::SetActorRotation(NewRotation, Teleport); } \
+	bool SetActorLocationAndRotation(FVector NewLocation, FRotator NewRotation, bool bSweep=false, FHitResult* OutSweepHitResult=nullptr, ETeleportType Teleport = ETeleportType::None) { return Super::SetActorLocationAndRotation(NewLocation, NewRotation, bSweep, OutSweepHitResult, Teleport); } \
+	bool SetActorLocationAndRotation(FVector NewLocation, const FQuat& NewRotation, bool bSweep=false, FHitResult* OutSweepHitResult=nullptr, ETeleportType Teleport = ETeleportType::None) { return Super::SetActorLocationAndRotation(NewLocation, NewRotation, bSweep, OutSweepHitResult, Teleport); } \
 	virtual bool TeleportTo( const FVector& DestLocation, const FRotator& DestRotation, bool bIsATest, bool bNoCheck ) override { return Super::TeleportTo(DestLocation, DestRotation, bIsATest, bNoCheck); } \
 	virtual FVector GetVelocity() const override { return Super::GetVelocity(); } \
 	float GetHorizontalDistanceTo(AActor* OtherActor)  { return Super::GetHorizontalDistanceTo(OtherActor); } \
 	float GetVerticalDistanceTo(AActor* OtherActor)  { return Super::GetVerticalDistanceTo(OtherActor); } \
 	float GetDotProductTo(AActor* OtherActor) { return Super::GetDotProductTo(OtherActor); } \
 	float GetHorizontalDotProductTo(AActor* OtherActor) { return Super::GetHorizontalDotProductTo(OtherActor); } \
-	float GetDistanceTo(AActor* OtherActor) { return Super::GetDistanceTo(OtherActor); }
-
-
-
+	float GetDistanceTo(AActor* OtherActor) { return Super::GetDistanceTo(OtherActor); } \
+	float GetSquaredDistanceTo(const AActor* OtherActor) { return Super::GetSquaredDistanceTo(OtherActor); } \
+	FVector GetActorForwardVector() const { return Super::GetActorForwardVector(); } \
+	FVector GetActorUpVector() const { return Super::GetActorUpVector(); } \
+	FVector GetActorRightVector() const { return Super::GetActorRightVector(); } \
+	void GetActorBounds(bool bOnlyCollidingComponents, FVector& Origin, FVector& BoxExtent) const { return Super::GetActorBounds(bOnlyCollidingComponents, Origin, BoxExtent); } \
+	void SetActorScale3D(FVector NewScale3D) { Super::SetActorScale3D(NewScale3D); } \
+	FVector GetActorScale3D() const { return Super::GetActorScale3D(); } \
+	void SetActorRelativeScale3D(FVector NewRelativeScale) { Super::SetActorRelativeScale3D(NewRelativeScale); } \
+	FVector GetActorRelativeScale3D() const { return Super::GetActorRelativeScale3D(); } \
+	FTransform ActorToWorld() const { return Super::ActorToWorld(); } \
+	void AddActorWorldOffset(FVector DeltaLocation, bool bSweep=false, FHitResult* OutSweepHitResult=nullptr, ETeleportType Teleport = ETeleportType::None) { Super::AddActorWorldOffset(DeltaLocation, bSweep, OutSweepHitResult, Teleport); } \
+	void AddActorWorldRotation(FRotator DeltaRotation, bool bSweep=false, FHitResult* OutSweepHitResult=nullptr, ETeleportType Teleport = ETeleportType::None) { Super::AddActorWorldRotation(DeltaRotation, bSweep, OutSweepHitResult, Teleport); } \
+	void AddActorWorldRotation(const FQuat& DeltaRotation, bool bSweep=false, FHitResult* OutSweepHitResult=nullptr, ETeleportType Teleport = ETeleportType::None) { Super::AddActorWorldRotation(DeltaRotation, bSweep, OutSweepHitResult, Teleport); } \
+	void AddActorWorldTransform(const FTransform& DeltaTransform, bool bSweep=false, FHitResult* OutSweepHitResult=nullptr, ETeleportType Teleport = ETeleportType::None) { Super::AddActorWorldTransform(DeltaTransform, bSweep, OutSweepHitResult, Teleport); } \
+	void AddActorLocalOffset(FVector DeltaLocation, bool bSweep=false, FHitResult* OutSweepHitResult=nullptr, ETeleportType Teleport = ETeleportType::None) { Super::AddActorLocalOffset(DeltaLocation, bSweep, OutSweepHitResult, Teleport); } \
+	void AddActorLocalRotation(FRotator DeltaRotation, bool bSweep=false, FHitResult* OutSweepHitResult=nullptr, ETeleportType Teleport = ETeleportType::None) { Super::AddActorLocalRotation(DeltaRotation, bSweep, OutSweepHitResult, Teleport); } \
+	void AddActorLocalRotation(const FQuat& DeltaRotation, bool bSweep=false, FHitResult* OutSweepHitResult=nullptr, ETeleportType Teleport = ETeleportType::None) { Super::AddActorLocalRotation(DeltaRotation, bSweep, OutSweepHitResult, Teleport); } \
+	void AddActorLocalTransform(const FTransform& NewTransform, bool bSweep=false, FHitResult* OutSweepHitResult=nullptr, ETeleportType Teleport = ETeleportType::None) { Super::AddActorLocalTransform(NewTransform, bSweep, OutSweepHitResult, Teleport); } \
+	void SetActorRelativeLocation(FVector NewRelativeLocation, bool bSweep=false, FHitResult* OutSweepHitResult=nullptr, ETeleportType Teleport = ETeleportType::None) { Super::SetActorRelativeLocation(NewRelativeLocation, bSweep, OutSweepHitResult, Teleport); } \
+	void SetActorRelativeRotation(FRotator NewRelativeRotation, bool bSweep=false, FHitResult* OutSweepHitResult=nullptr, ETeleportType Teleport = ETeleportType::None) { Super::SetActorRelativeRotation(NewRelativeRotation, bSweep, OutSweepHitResult, Teleport); } \
+	void SetActorRelativeRotation(const FQuat& NewRelativeRotation, bool bSweep=false, FHitResult* OutSweepHitResult=nullptr, ETeleportType Teleport = ETeleportType::None) { Super::SetActorRelativeRotation(NewRelativeRotation, bSweep, OutSweepHitResult, Teleport); } \
+	void SetActorRelativeTransform(const FTransform& NewRelativeTransform, bool bSweep=false, FHitResult* OutSweepHitResult=nullptr, ETeleportType Teleport = ETeleportType::None) { Super::SetActorRelativeTransform(NewRelativeTransform, bSweep, OutSweepHitResult, Teleport); }
diff --git a/Engine/Source/Runtime/Engine/Private/Actor.cpp b/Engine/Source/Runtime/Engine/Private/Actor.cpp
--- a/Engine/Source/Runtime/Engine/Private/Actor.cpp
+++ b/Engine/Source/Runtime/Engine/Private/Actor.cpp
@@ -128,8 +128,11 @@ void FActorTickFunction::ExecuteTick(float DeltaTime, enum ELevelTick TickType,
 {
 	if (Target && !Target->IsPendingKillOrUnreachable())
 	{
-		FScopeCycleCounterUObject ActorScope(Target);
-		Target->TickActor(DeltaTime*Target->CustomTimeDilation, TickType, *this);	
+		if (TickType != LEVELTICK_ViewportsOnly || Target->ShouldTickIfViewportsOnly())
+		{
+			FScopeCycleCounterUObject ActorScope(Target);
+			Target->TickActor(DeltaTime*Target->CustomTimeDilation, TickType, *this);
+		}
 	}
 }
 
@@ -809,15 +812,17 @@ float AActor::GetActorTickInterval() const
 
 bool AActor::Rename( const TCHAR* InName, UObject* NewOuter, ERenameFlags Flags )
 {
-	if (NewOuter)
+	const bool bRenameTest = ((Flags & REN_Test) != 0);
+
+	if (!bRenameTest && NewOuter)
 	{
 		RegisterAllActorTickFunctions(false, true); // unregister all tick functions
 		UnregisterAllComponents();
 	}
 
-	bool bSuccess = Super::Rename( InName, NewOuter, Flags );
+	const bool bSuccess = Super::Rename( InName, NewOuter, Flags );
 
-	if (NewOuter && NewOuter->IsA<ULevel>())
+	if (!bRenameTest && NewOuter && NewOuter->IsA<ULevel>())
 	{
 		UWorld* World = NewOuter->GetWorld();
 		if (World && World->bIsWorldInitialized)
@@ -857,14 +862,10 @@ void AActor::TickActor( float DeltaSeconds, ELevelTick TickType, FActorTickFunct
 	//root of tick hierarchy
 
 	// Non-player update.
-	const bool bShouldTick = ((TickType!=LEVELTICK_ViewportsOnly) || ShouldTickIfViewportsOnly());
-	if(bShouldTick)
+	// If an Actor has been Destroyed or its level has been unloaded don't execute any queued ticks
+	if (!IsPendingKill() && GetWorld())
 	{
-		// If an Actor has been Destroyed or its level has been unloaded don't execute any queued ticks
-		if (!IsPendingKill() && GetWorld())
-		{
-			Tick(DeltaSeconds);	// perform any tick functions unique to an actor subclass
-		}
+		Tick(DeltaSeconds);	// perform any tick functions unique to an actor subclass
 	}
 }
 
diff --git a/Engine/Source/Runtime/Engine/Private/ComponentInstanceDataCache.cpp b/Engine/Source/Runtime/Engine/Private/ComponentInstanceDataCache.cpp
--- a/Engine/Source/Runtime/Engine/Private/ComponentInstanceDataCache.cpp
+++ b/Engine/Source/Runtime/Engine/Private/ComponentInstanceDataCache.cpp
@@ -136,7 +136,7 @@ public:
 		InComponent->GetUCSModifiedProperties(PropertiesToSkip);
 
 		UClass* Class = InComponent->GetClass();
-		Class->SerializeTaggedProperties(*this, (uint8*)InComponent, Class, nullptr);
+		Class->SerializeTaggedProperties(*this, (uint8*)InComponent, Class, (uint8*)InComponent->GetArchetype());
 	}
 
 	virtual bool ShouldSkipProperty(const UProperty* InProperty) const override
diff --git a/Engine/Source/Runtime/Engine/Private/Components/ChildActorComponent.cpp b/Engine/Source/Runtime/Engine/Private/Components/ChildActorComponent.cpp
--- a/Engine/Source/Runtime/Engine/Private/Components/ChildActorComponent.cpp
+++ b/Engine/Source/Runtime/Engine/Private/Components/ChildActorComponent.cpp
@@ -378,7 +378,7 @@ void UChildActorComponent::ApplyComponentInstanceData(FChildActorComponentInstan
 			const FString ChildActorNameString = ChildActorName.ToString();
 			if (ChildActor->Rename(*ChildActorNameString, nullptr, REN_Test))
 			{
-				ChildActor->Rename(*ChildActorNameString, nullptr, REN_DoNotDirty | (IsLoading() ? REN_ForceNoResetLoaders : REN_None));
+				ChildActor->Rename(*ChildActorNameString, nullptr, REN_DoNotDirty | REN_ForceNoResetLoaders);
 			}
 		}
 
@@ -425,7 +425,9 @@ void UChildActorComponent::SetChildActorClass(TSubclassOf<AActor> Class)
 				if (ChildActorTemplate)
 				{
 					UEngine::CopyPropertiesForUnrelatedObjects(ChildActorTemplate, NewChildActorTemplate);
-
+#if WITH_EDITOR
+					NewChildActorTemplate->ClearActorLabel();
+#endif
 					ChildActorTemplate->Rename(nullptr, GetTransientPackage(), REN_DontCreateRedirectors);
 				}
 
@@ -631,8 +633,7 @@ void UChildActorComponent::DestroyChildActor()
 				if (!IsGarbageCollecting())
 				{
 					const FString ObjectBaseName = FString::Printf(TEXT("DESTROYED_%s_CHILDACTOR"), *ChildClass->GetName());
-					const ERenameFlags RenameFlags = ((GetWorld()->IsGameWorld() || IsLoading()) ? REN_DoNotDirty | REN_ForceNoResetLoaders : REN_DoNotDirty);
-					ChildActor->Rename(*MakeUniqueObjectName(ChildActor->GetOuter(), ChildClass, *ObjectBaseName).ToString(), nullptr, RenameFlags);
+					ChildActor->Rename(*MakeUniqueObjectName(ChildActor->GetOuter(), ChildClass, *ObjectBaseName).ToString(), nullptr, REN_DoNotDirty | REN_ForceNoResetLoaders);
 				}
 				else
 				{
diff --git a/Engine/Source/Runtime/Engine/Private/Components/SceneComponent.cpp b/Engine/Source/Runtime/Engine/Private/Components/SceneComponent.cpp
--- a/Engine/Source/Runtime/Engine/Private/Components/SceneComponent.cpp
+++ b/Engine/Source/Runtime/Engine/Private/Components/SceneComponent.cpp
@@ -378,13 +378,16 @@ static int32 SetDescendantIsEditorOnly(USceneComponent const* SceneComponentObje
 	// recursively alter the bIsEditorOnly flag for children and deeper descendants 
 	for (USceneComponent* ChildSceneComponent : AttachedChildren)
 	{
-		if (!ChildSceneComponent->bIsEditorOnly)
+		if (ChildSceneComponent != nullptr)
 		{
-			ChildSceneComponent->Modify();
-			ChildSceneComponent->bIsEditorOnly = true;
-			++NumDescendantsChanged;
+			if (!ChildSceneComponent->bIsEditorOnly)
+			{
+				ChildSceneComponent->Modify();
+				ChildSceneComponent->bIsEditorOnly = true;
+				++NumDescendantsChanged;
+			}
+			NumDescendantsChanged += SetDescendantIsEditorOnly(ChildSceneComponent);
 		}
-		NumDescendantsChanged += SetDescendantIsEditorOnly(ChildSceneComponent);
 	}
 
 	return NumDescendantsChanged;
@@ -841,43 +844,44 @@ void USceneComponent::DestroyComponent(bool bPromoteChildren/*= false*/)
 			{
 				// Cache our AttachParent
 				USceneComponent* CachedAttachParent = GetAttachParent();
-				check(CachedAttachParent != nullptr);
-
-				// Find the our position in its AttachParent's child array
-				const TArray<USceneComponent*>& AttachSiblings = CachedAttachParent->GetAttachChildren();
-				int32 Index = AttachSiblings.Find(this);
-				check(Index != INDEX_NONE);
-
-				// Detach from parent
-				DetachFromComponent(FDetachmentTransformRules::KeepWorldTransform);
-
-				// Find an appropriate child node to promote to this node's position in the hierarchy
-				if (AttachedChildren.Num() > 0)
+				if (ensureMsgf(CachedAttachParent != nullptr, TEXT("Deleting a non-root scene component with no AttachParent: %s"), *GetFullName()))
 				{
-					// Always choose non editor-only child nodes over editor-only child nodes (since we don't want editor-only nodes to end up with non editor-only child nodes)
-					USceneComponent* const * FindResult =
-						AttachedChildren.FindByPredicate([Owner](USceneComponent* Child){ return Child != nullptr && !Child->IsEditorOnly(); });
+					// Find the our position in its AttachParent's child array
+					const TArray<USceneComponent*>& AttachSiblings = CachedAttachParent->GetAttachChildren();
+					int32 Index = AttachSiblings.Find(this);
+					check(Index != INDEX_NONE);
 
-					if (FindResult != nullptr)
-					{
-						ChildToPromote = *FindResult;
-					}
-					else
+					// Detach from parent
+					DetachFromComponent(FDetachmentTransformRules::KeepWorldTransform);
+
+					// Find an appropriate child node to promote to this node's position in the hierarchy
+					if (AttachedChildren.Num() > 0)
 					{
-						// Default to first child node
-						check(AttachedChildren[0] != nullptr);
-						ChildToPromote = AttachedChildren[0];
+						// Always choose non editor-only child nodes over editor-only child nodes (since we don't want editor-only nodes to end up with non editor-only child nodes)
+						USceneComponent* const * FindResult =
+							AttachedChildren.FindByPredicate([Owner](USceneComponent* Child) { return Child != nullptr && !Child->IsEditorOnly(); });
+
+						if (FindResult != nullptr)
+						{
+							ChildToPromote = *FindResult;
+						}
+						else
+						{
+							// Default to first child node
+							check(AttachedChildren[0] != nullptr);
+							ChildToPromote = AttachedChildren[0];
+						}
 					}
-				}
 
-				if (ChildToPromote != nullptr)
-				{
-					// Attach the child node that we're promoting to the parent and move it to the same position as the old node was in the array
-					ChildToPromote->AttachToComponent(CachedAttachParent, FAttachmentTransformRules::KeepWorldTransform);
-					CachedAttachParent->AttachChildren.Remove(ChildToPromote);
+					if (ChildToPromote != nullptr)
+					{
+						// Attach the child node that we're promoting to the parent and move it to the same position as the old node was in the array
+						ChildToPromote->AttachToComponent(CachedAttachParent, FAttachmentTransformRules::KeepWorldTransform);
+						CachedAttachParent->AttachChildren.Remove(ChildToPromote);
 
-					Index = FMath::Clamp<int32>(Index, 0, AttachSiblings.Num());
-					CachedAttachParent->AttachChildren.Insert(ChildToPromote, Index);
+						Index = FMath::Clamp<int32>(Index, 0, AttachSiblings.Num());
+						CachedAttachParent->AttachChildren.Insert(ChildToPromote, Index);
+					}
 				}
 			}
 
@@ -885,13 +889,14 @@ void USceneComponent::DestroyComponent(bool bPromoteChildren/*= false*/)
 			TArray<USceneComponent*> AttachChildrenLocalCopy(AttachedChildren);
 			for (USceneComponent* Child : AttachChildrenLocalCopy)
 			{
-				check(Child != nullptr);
-
-				// Note: This will internally call Modify(), so we don't need to call it here
-				Child->DetachFromComponent(FDetachmentTransformRules::KeepWorldTransform);
-				if (Child != ChildToPromote)
+				if (ensure(Child))
 				{
-					Child->AttachToComponent(ChildToPromote, FAttachmentTransformRules::KeepWorldTransform);
+					// Note: This will internally call Modify(), so we don't need to call it here
+					Child->DetachFromComponent(FDetachmentTransformRules::KeepWorldTransform);
+					if (Child != ChildToPromote)
+					{
+						Child->AttachToComponent(ChildToPromote, FAttachmentTransformRules::KeepWorldTransform);
+					}
 				}
 			}
 		}
diff --git a/Engine/Source/Runtime/Engine/Private/DataTableCSV.cpp b/Engine/Source/Runtime/Engine/Private/DataTableCSV.cpp
--- a/Engine/Source/Runtime/Engine/Private/DataTableCSV.cpp
+++ b/Engine/Source/Runtime/Engine/Private/DataTableCSV.cpp
@@ -170,11 +170,6 @@ bool FDataTableImporterCSV::ReadTable()
 		DataTable->RowStruct->InitializeStruct(RowData);
 		// And be sure to call DestroyScriptStruct later
 
-		if (auto UDStruct = Cast<const UUserDefinedStruct>(DataTable->RowStruct))
-		{
-			UDStruct->InitializeDefaultValue(RowData);
-		}
-
 		// Add to row map
 		DataTable->RowMap.Add(RowName, RowData);
 
diff --git a/Engine/Source/Runtime/Engine/Private/DataTableJSON.cpp b/Engine/Source/Runtime/Engine/Private/DataTableJSON.cpp
--- a/Engine/Source/Runtime/Engine/Private/DataTableJSON.cpp
+++ b/Engine/Source/Runtime/Engine/Private/DataTableJSON.cpp
@@ -445,11 +445,6 @@ bool FDataTableImporterJSON::ReadRow(const TSharedRef<FJsonObject>& InParsedTabl
 	DataTable->RowStruct->InitializeStruct(RowData);
 	// And be sure to call DestroyScriptStruct later
 
-	if (auto UDStruct = Cast<const UUserDefinedStruct>(DataTable->RowStruct))
-	{
-		UDStruct->InitializeDefaultValue(RowData);
-	}
-
 	// Add to row map
 	DataTable->RowMap.Add(RowName, RowData);
 
diff --git a/Engine/Source/Runtime/Engine/Private/EdGraph/EdGraphPin.cpp b/Engine/Source/Runtime/Engine/Private/EdGraph/EdGraphPin.cpp
--- a/Engine/Source/Runtime/Engine/Private/EdGraph/EdGraphPin.cpp
+++ b/Engine/Source/Runtime/Engine/Private/EdGraph/EdGraphPin.cpp
@@ -18,6 +18,10 @@
 
 #define LOCTEXT_NAMESPACE "EdGraph"
 
+#if WITH_EDITOR
+extern UNREALED_API UEditorEngine* GEditor;
+#endif
+
 class FPinDeletionQueue : 
 #if WITH_EDITOR
 	public FTickableEditorObject
@@ -71,6 +75,60 @@ private:
 
 };
 
+struct FGraphConnectionSanitizer
+{
+	void AddCorruptGraph(UEdGraph* CorruptGraph);
+	void SanitizeConnectionsForCorruptGraphs();
+
+	TArray<UEdGraph*> CorruptGraphs;
+} GraphConnectionSanitizer;
+
+void FGraphConnectionSanitizer::AddCorruptGraph(UEdGraph* CorruptGraph)
+{
+	CorruptGraphs.Add(CorruptGraph);
+}
+
+void FGraphConnectionSanitizer::SanitizeConnectionsForCorruptGraphs()
+{
+	for(UEdGraph* CorruptGraph : CorruptGraphs)
+	{
+		// validate all links for nodes in this graph:
+		for(UEdGraphNode* Node : CorruptGraph->Nodes)
+		{
+			TArray<UEdGraphPin*>& NodePins = Node->Pins;
+			for(int32 I = NodePins.Num() - 1; I >= 0; --I )
+			{
+				UEdGraphPin* CurrentPin = NodePins[I];
+				// check for null or assymetry:
+				if(CurrentPin == nullptr)
+				{
+					NodePins.RemoveAtSwap(I);
+				}
+				else
+				{
+					TArray<UEdGraphPin*>& PinLinkedTo = CurrentPin->LinkedTo;
+					for( int32 J = PinLinkedTo.Num() - 1; J >= 0; --J )
+					{
+						UEdGraphPin* CurrentLinkedToPin = PinLinkedTo[J];
+						if(CurrentLinkedToPin == nullptr)
+						{
+							PinLinkedTo.RemoveAtSwap(J);
+						}
+						else
+						{
+							if(!CurrentLinkedToPin->LinkedTo.Contains(CurrentPin))
+							{
+								PinLinkedTo.RemoveAtSwap(J);
+							}
+						}
+					}
+				}
+			}
+		}
+	}
+	CorruptGraphs.Empty();
+}
+
 //#define TRACK_PINS
 
 #ifdef TRACK_PINS
@@ -194,14 +252,10 @@ bool FEdGraphPinType::Serialize(FArchive& Ar)
 				PinSubCategoryMemberReference.MemberParent = Signature->GetOwnerClass();
 				PinSubCategoryObject = NULL;
 			}
-			else
-			{
-				ensure(true);
-			}
 		}
 	}
 
-	bool bIsConstBool = false;
+	bool bIsConstBool = bIsConst;
 
 	if (Ar.UE4Ver() >= VER_UE4_SERIALIZE_PINTYPE_CONST)
 	{
@@ -1668,7 +1722,7 @@ void UEdGraphPin::ResolveReferencesToPin(UEdGraphPin* Pin, bool bStrictValidatio
 					// When in the middle of a transaction the LinkedTo lists will be in an 
 					// indeterminate state. After PostEditUndo runs we could validate LinkedTo
 					// coherence.
-					ensureAlways(GIsTransacting || Pin->LinkedTo.Contains(ReferencingPin));
+					ensureAlways(Pin->LinkedTo.Contains(ReferencingPin) || GEditor->Trans->IsObjectTransacting(Pin->GetOwningNode()));
 #else
 					ensureAlways(Pin->LinkedTo.Contains(ReferencingPin));
 #endif//WITH_EDITOR
@@ -1708,6 +1762,24 @@ void UEdGraphPin::SerializePinArray(FArchive& Ar, TArray<UEdGraphPin*>& ArrayRef
 
 	if (Ar.IsLoading())
 	{
+		if(Ar.IsTransacting() && ResolveType == EPinResolveType::LinkedTo)
+		{
+			if(ArrayNum < ArrayRef.Num())
+			{
+				for(int32 I = ArrayRef.Num() - 1; I >=0 ; --I)
+				{
+					if( !ArrayRef[I]->bWasTrashed &&
+#if WITH_EDITOR
+						!GEditor->Trans->IsObjectTransacting(ArrayRef[I]->GetOwningNode()) &&
+#endif
+						ArrayRef[I]->LinkedTo.Contains(RequestingPin) )
+					{
+						ArrayRef[I]->LinkedTo.Remove(RequestingPin);
+					}
+				}
+			}
+		}
+
 		ArrayRef.SetNum(ArrayNum);
 	}
 
@@ -1726,7 +1798,6 @@ void UEdGraphPin::SerializePinArray(FArchive& Ar, TArray<UEdGraphPin*>& ArrayRef
 #if WITH_EDITOR
 			// More complexity to handle asymmetry in the transaction buffer. If our peer node is not in the transaction
 			// then we need to take ownership of the entire connection and clear both LinkedTo Arrays:
-			extern UNREALED_API UEditorEngine* GEditor;
 			for (int32 I = 0; I < Pin->LinkedTo.Num(); ++I)
 			{
 				UEdGraphPin* Peer = Pin->LinkedTo[I];
@@ -1808,16 +1879,17 @@ bool UEdGraphPin::SerializePin(FArchive& Ar, UEdGraphPin*& PinRef, int32 ArrayId
 				{
 					check(RequestingPin);
 #if WITH_EDITOR
-					extern UNREALED_API UEditorEngine* GEditor;
 					if(Ar.IsTransacting() && !GEditor->Trans->IsObjectTransacting(LocalOwningNode))
 					{
 						/* 
-							Two possibilities:
+							Three possibilities:
 
 							1. This transaction did not alter this node's LinkedTo list, and also did not
 								alter ExistingPin's
 							2. This transaction has altered this node's LinkedTo list, but did not alter
 								ExistingPin's
+							3. Mutations to connections have been made outside the transaction system, corrupting
+								the transaction buffer
 
 							In case 2 we need to make sure that the LinkedTo connection is symmetrical, but
 							in the case of 1 we have to assume that there is already a connection in ExistingPin.
@@ -1827,20 +1899,35 @@ bool UEdGraphPin::SerializePin(FArchive& Ar, UEdGraphPin*& PinRef, int32 ArrayId
 
 						check(ResolveType == EPinResolveType::LinkedTo);
 
-						check(ExistingPin && *ExistingPin); // the transaction buffer is corrupt
-						FGuid RequestingPinId = RequestingPin->PinId;
-						UEdGraphPin** LinkedTo = (*ExistingPin)->LinkedTo.FindByPredicate([&RequestingPinId, RequestingPin](const UEdGraphPin* Pin) { return Pin == RequestingPin && Pin->PinId == RequestingPinId; });
-						if (LinkedTo)
+						if(ExistingPin && *ExistingPin)
 						{
-							// case 1:
-							PinRef = *ExistingPin;
+							FGuid RequestingPinId = RequestingPin->PinId;
+							UEdGraphPin** LinkedTo = (*ExistingPin)->LinkedTo.FindByPredicate([&RequestingPinId, RequestingPin](const UEdGraphPin* Pin) { return Pin == RequestingPin && Pin->PinId == RequestingPinId; });
+							if (LinkedTo)
+							{
+								// case 1:
+								PinRef = *ExistingPin;
+								ensureAlways(!(*ExistingPin)->LinkedTo.Contains(PinRef) || GEditor->Trans->IsObjectTransacting((*ExistingPin)->GetOwningNode()));
+							}
+							else
+							{
+								// case 2:
+								TArray<FUnresolvedPinData>& UnresolvedPinData = PinHelpers::UnresolvedPins.FindOrAdd(FPinResolveId(PinGuid, LocalOwningNode));
+								UnresolvedPinData.Add(FUnresolvedPinData(RequestingPin, ResolveType, ArrayIdx, true));
+								bRetVal = false;
+							}
 						}
 						else
 						{
-							// case 2:
-							TArray<FUnresolvedPinData>& UnresolvedPinData = PinHelpers::UnresolvedPins.FindOrAdd(FPinResolveId(PinGuid, LocalOwningNode));
-							UnresolvedPinData.Add(FUnresolvedPinData(RequestingPin, ResolveType, ArrayIdx, true));
+							// case 3:
+							UE_LOG(LogBlueprint, Warning, TEXT("Pin link in transaction buffer is corrupt - this occurs because node mutations are not recorded in the transaction buffer. Node links will be lost"));
+							PinRef = nullptr;
 							bRetVal = false;
+
+							// After the transaction is complete, we'll want to fix any asymmetries in 
+							// pin links, this will be somewhat costly, but a small hitch is much
+							// better than crashing:
+							GraphConnectionSanitizer.AddCorruptGraph(LocalOwningNode->GetGraph());
 						}
 					}
 					else
@@ -2008,6 +2095,11 @@ bool UEdGraphPin::AreOrphanPinsEnabled()
 	return (GCVarDisableOrphanPins == 0);
 }
 
+void UEdGraphPin::SanitizePinsPostUndoRedo()
+{
+	ensure(PinHelpers::UnresolvedPins.Num() == 0);
+	GraphConnectionSanitizer.SanitizeConnectionsForCorruptGraphs();
+}
 
 /////////////////////////////////////////////////////
 // UEdGraphPin
diff --git a/Engine/Source/Runtime/Engine/Private/EdGraph/EdGraphSchema.cpp b/Engine/Source/Runtime/Engine/Private/EdGraph/EdGraphSchema.cpp
--- a/Engine/Source/Runtime/Engine/Private/EdGraph/EdGraphSchema.cpp
+++ b/Engine/Source/Runtime/Engine/Private/EdGraph/EdGraphSchema.cpp
@@ -693,7 +693,15 @@ FText UEdGraphSchema::GetPinDisplayName(const UEdGraphPin* Pin) const
 	check(Pin != nullptr);
 	if (Pin->PinFriendlyName.IsEmpty())
 	{
-		ResultPinName = FText::FromName(Pin->PinName);
+		// We don't want to display "None" for no name
+		if (Pin->PinName.IsNone())
+		{
+			return FText::GetEmpty();
+		}
+		else
+		{
+			ResultPinName = FText::FromName(Pin->PinName);
+		}
 	}
 	else
 	{
diff --git a/Engine/Source/Runtime/Engine/Private/NetworkDriver.cpp b/Engine/Source/Runtime/Engine/Private/NetworkDriver.cpp
--- a/Engine/Source/Runtime/Engine/Private/NetworkDriver.cpp
+++ b/Engine/Source/Runtime/Engine/Private/NetworkDriver.cpp
@@ -1943,6 +1943,8 @@ bool UNetDriver::HandleNetDumpServerRPCCommand( const TCHAR* Cmd, FOutputDevice&
 	{
 		bool bHasNetFields = false;
 
+		ensureMsgf(!ClassIt->HasAnyFlags(RF_NeedLoad | RF_NeedPostLoad), TEXT("UNetDriver::HandleNetDumpServerRPCCommand: %s has flag RF_NeedPostLoad. NetFields and ClassReps will be incorrect!"), *GetFullNameSafe(*ClassIt));
+
 		for ( int32 i = 0; i < ClassIt->NetFields.Num(); i++ )
 		{
 			UFunction * Function = Cast<UFunction>( ClassIt->NetFields[i] );
diff --git a/Engine/Source/Runtime/Engine/Private/RepLayout.cpp b/Engine/Source/Runtime/Engine/Private/RepLayout.cpp
--- a/Engine/Source/Runtime/Engine/Private/RepLayout.cpp
+++ b/Engine/Source/Runtime/Engine/Private/RepLayout.cpp
@@ -2971,6 +2971,8 @@ void FRepLayout::InitFromObjectClass( UClass * InObjectClass )
 
 	Parents.Empty();
 
+	ensureMsgf(!InObjectClass->HasAnyFlags(RF_NeedLoad | RF_NeedPostLoad), TEXT("FRepLayout::InitFromObjectClass: %s has not. NetFields and ClassReps will be incorrect!"), *GetFullNameSafe(InObjectClass));
+
 	for ( int32 i = 0; i < InObjectClass->ClassReps.Num(); i++ )
 	{
 		UProperty * Property	= InObjectClass->ClassReps[i].Property;
diff --git a/Engine/Source/Runtime/Engine/Private/UserDefinedStruct.cpp b/Engine/Source/Runtime/Engine/Private/UserDefinedStruct.cpp
--- a/Engine/Source/Runtime/Engine/Private/UserDefinedStruct.cpp
+++ b/Engine/Source/Runtime/Engine/Private/UserDefinedStruct.cpp
@@ -9,13 +9,44 @@
 #include "UObject/FrameworkObjectVersion.h"
 #include "Misc/SecureHash.h"
 #include "UObject/PropertyPortFlags.h"
-#include "Misc/PackageName.h" // for FPackageName::GetLongPackageAssetName()
+#include "Misc/PackageName.h"
+#include "Serialization/TextReferenceCollector.h"
+#include "Blueprint/BlueprintSupport.h"
 
 #if WITH_EDITOR
 #include "UserDefinedStructure/UserDefinedStructEditorData.h"
 #include "Kismet2/StructureEditorUtils.h"
 #endif //WITH_EDITOR
-#include "Serialization/TextReferenceCollector.h"
+
+FUserStructOnScopeIgnoreDefaults::FUserStructOnScopeIgnoreDefaults(const UUserDefinedStruct* InUserStruct)
+{
+	// Can't call super constructor because we need to call our overridden initialize
+	ScriptStruct = InUserStruct;
+	SampleStructMemory = nullptr;
+	OwnsMemory = false;
+	Initialize();
+}
+
+FUserStructOnScopeIgnoreDefaults::FUserStructOnScopeIgnoreDefaults(const UUserDefinedStruct* InUserStruct, uint8* InData) : FStructOnScope(InUserStruct, InData)
+{
+}
+
+void FUserStructOnScopeIgnoreDefaults::Recreate(const UUserDefinedStruct* InUserStruct)
+{
+	Destroy();
+	ScriptStruct = InUserStruct;
+	Initialize();
+}
+
+void FUserStructOnScopeIgnoreDefaults::Initialize()
+{
+	if (ScriptStruct.IsValid())
+	{
+		SampleStructMemory = (uint8*)FMemory::Malloc(ScriptStruct->GetStructureSize() ? ScriptStruct->GetStructureSize() : 1);
+		((UUserDefinedStruct*)ScriptStruct.Get())->InitializeStructIgnoreDefaults(SampleStructMemory);
+		OwnsMemory = true;
+	}
+}
 
 #if WITH_EDITORONLY_DATA
 namespace
@@ -28,10 +59,7 @@ namespace
 
 		const FString PathToObject = UserDefinedStruct->GetPathName();
 
-		FStructOnScope StructData(UserDefinedStruct);
-		FStructureEditorUtils::Fill_MakeStructureDefaultValue(UserDefinedStruct, StructData.GetStructMemory());
-
-		PropertyLocalizationDataGatherer.GatherLocalizationDataFromStructFields(PathToObject, StructData.GetStruct(), StructData.GetStructMemory(), nullptr, GatherTextFlags);
+		PropertyLocalizationDataGatherer.GatherLocalizationDataFromStructFields(PathToObject, UserDefinedStruct, UserDefinedStruct->GetDefaultInstance(), nullptr, GatherTextFlags);
 	}
 
 	void CollectUserDefinedStructTextReferences(UObject* Object, FArchive& Ar)
@@ -64,20 +92,43 @@ namespace
 UUserDefinedStruct::UUserDefinedStruct(const FObjectInitializer& ObjectInitializer)
 	: Super(ObjectInitializer)
 {
+	DefaultStructInstance.SetPackage(GetOutermost());
+
 #if WITH_EDITORONLY_DATA
 	{ static const FAutoRegisterLocalizationDataGatheringCallback AutomaticRegistrationOfLocalizationGatherer(UUserDefinedStruct::StaticClass(), &GatherUserDefinedStructForLocalization); }
 	{ static const FAutoRegisterTextReferenceCollectorCallback AutomaticRegistrationOfTextReferenceCollector(UUserDefinedStruct::StaticClass(), &CollectUserDefinedStructTextReferences); }
 #endif
 }
 
-#if WITH_EDITOR
-
 void UUserDefinedStruct::Serialize(FArchive& Ar)
 {
 	Super::Serialize( Ar );
 
+	if (HasAnyFlags(RF_ClassDefaultObject))
+	{
+		return;
+	}
+
 	Ar.UsingCustomVersion(FFrameworkObjectVersion::GUID);
 
+	if (Ar.CustomVer(FFrameworkObjectVersion::GUID) >= FFrameworkObjectVersion::UserDefinedStructsStoreDefaultInstance)
+	{
+		if (EUserDefinedStructureStatus::UDSS_UpToDate == Status && !(Ar.GetPortFlags() & PPF_Duplicate))
+		{
+			// If we're saving or loading new data, serialize our defaults
+			if (!DefaultStructInstance.IsValid() && Ar.IsLoading())
+			{
+				DefaultStructInstance.Recreate(this);
+			}
+
+			uint8* StructData = DefaultStructInstance.GetStructMemory();
+
+			FScopedPlaceholderRawContainerTracker TrackStruct(StructData);
+			SerializeItem(Ar, StructData, nullptr);
+		}
+	}
+
+#if WITH_EDITOR
 	if (Ar.IsLoading())
 	{
 		if (Ar.CustomVer(FFrameworkObjectVersion::GUID) < FFrameworkObjectVersion::UserDefinedStructsBlueprintVisible)
@@ -109,8 +160,11 @@ void UUserDefinedStruct::Serialize(FArchive& Ar)
 			}
 		}
 	}
+#endif
 }
 
+#if WITH_EDITOR
+
 void UUserDefinedStruct::PostDuplicate(bool bDuplicateForPIE)
 {
 	Super::PostDuplicate(bDuplicateForPIE);
@@ -142,21 +196,11 @@ void UUserDefinedStruct::GetAssetRegistryTags(TArray<FAssetRegistryTag>& OutTags
 UProperty* UUserDefinedStruct::CustomFindProperty(const FName Name) const
 {
 	const FGuid PropertyGuid = FStructureEditorUtils::GetGuidFromPropertyName(Name);
-	UProperty* Property = PropertyGuid.IsValid() ? FStructureEditorUtils::GetPropertyByGuid(this, PropertyGuid) : NULL;
-	ensure(!Property || PropertyGuid == FStructureEditorUtils::GetGuidForProperty(Property));
+	UProperty* Property = PropertyGuid.IsValid() ? FStructureEditorUtils::GetPropertyByGuid(this, PropertyGuid) : FStructureEditorUtils::GetPropertyByDisplayName(this, Name.ToString());
+	ensure(!Property || !PropertyGuid.IsValid() || PropertyGuid == FStructureEditorUtils::GetGuidForProperty(Property));
 	return Property;
 }
 
-void UUserDefinedStruct::InitializeDefaultValue(uint8* StructData) const
-{
-	FStructureEditorUtils::Fill_MakeStructureDefaultValue(this, StructData);
-}
-
-bool UUserDefinedStruct::DiffersFromDefaultValue(uint8* StructData) const
-{
-	return FStructureEditorUtils::DiffersFromDefaultValue(this, StructData);
-}
-
 void UUserDefinedStruct::ValidateGuid()
 {
 	// Backward compatibility:
@@ -198,51 +242,72 @@ FString UUserDefinedStruct::PropertyNameToDisplayName(FName Name) const
 	return OriginalName;
 }
 
+void UUserDefinedStruct::InitializeStructIgnoreDefaults(void* Dest, int32 ArrayDim) const
+{
+	Super::InitializeStruct(Dest, ArrayDim);
+}
+
+void UUserDefinedStruct::InitializeStruct(void* Dest, int32 ArrayDim) const
+{
+	InitializeStructIgnoreDefaults(Dest, ArrayDim);
+
+	if (Dest)
+	{
+		const uint8* DefaultInstance = GetDefaultInstance();
+		if (DefaultInstance)
+		{
+			int32 Stride = GetStructureSize();
+
+			for (int32 ArrayIndex = 0; ArrayIndex < ArrayDim; ArrayIndex++)
+			{
+				void* DestStruct = (uint8*)Dest + (Stride * ArrayIndex);
+				CopyScriptStruct(DestStruct, DefaultInstance);
+
+				// When copying into another struct we need to register this raw struct pointer so any deferred dependencies will get fixed later
+				FScopedPlaceholderRawContainerTracker TrackStruct(DestStruct);
+				FBlueprintSupport::RegisterDeferredDependenciesInStruct(this, DestStruct);
+			}	
+		}
+	}
+}
+
 void UUserDefinedStruct::SerializeTaggedProperties(FArchive& Ar, uint8* Data, UStruct* DefaultsStruct, uint8* Defaults, const UObject* BreakRecursionIfFullyLoad) const
 {
-#if WITH_EDITOR
-	/*	The following code is responsible for UUserDefinedStruct's default values serialization.	*/
+	bool bTemporarilyEnableDelta = false;
 
-	auto UDDefaultsStruct = Cast<UUserDefinedStruct>(DefaultsStruct);
+#if WITH_EDITOR
+	// In the editor the default structure may change while the editor is running, so we need to always delta serialize
+	UUserDefinedStruct* UDDefaultsStruct = Cast<UUserDefinedStruct>(DefaultsStruct);
 
 	const bool bDuplicate = (0 != (Ar.GetPortFlags() & PPF_Duplicate));
 
-	/*	When saving delta, we want the difference between current data and true structure's default values. 
-		When Defaults is NULL then zeroed data will be used for comparison.*/
+	// When saving delta, we want the difference between current data and true structure's default values.
+	// So if we don't have defaults we need to use the struct defaults
 	const bool bUseNewDefaults = !Defaults
 		&& UDDefaultsStruct
-		&& Ar.DoDelta()
-		&& Ar.IsSaving()
 		&& !bDuplicate
+		&& (Ar.IsSaving() || Ar.IsLoading())
 		&& !Ar.IsCooking();
 
-	/*	Object serialized from delta will have missing properties filled with zeroed data, 
-		we want structure's default data instead */
-	const bool bLoadDefaultFirst = UDDefaultsStruct
-		&& !bDuplicate
-		&& Ar.IsLoading();
-
-	const bool bPrepareDefaultStruct = bUseNewDefaults || bLoadDefaultFirst;
-	FStructOnScope StructDefaultMem(bPrepareDefaultStruct ? UDDefaultsStruct : NULL);
-	if (bPrepareDefaultStruct)
-	{
-		FStructureEditorUtils::Fill_MakeStructureDefaultValue(UDDefaultsStruct, StructDefaultMem.GetStructMemory());
-	}
-
 	if (bUseNewDefaults)
 	{
-		Defaults = StructDefaultMem.GetStructMemory();
+		Defaults = const_cast<uint8*>(UDDefaultsStruct->GetDefaultInstance());
 	}
-	if (bLoadDefaultFirst)
+
+	// Temporarily enable delta serialization if this is a CPFUO 
+	bTemporarilyEnableDelta = bUseNewDefaults && Ar.IsIgnoringArchetypeRef() && Ar.IsIgnoringClassRef() && !Ar.DoDelta();
+	if (bTemporarilyEnableDelta)
 	{
-		if (Defaults == nullptr)
-		{
-			Defaults = StructDefaultMem.GetStructMemory();
-		}
-		UDDefaultsStruct->CopyScriptStruct(Data, Defaults);
+		Ar.ArNoDelta = false;
 	}
 #endif // WITH_EDITOR
+
 	Super::SerializeTaggedProperties(Ar, Data, DefaultsStruct, Defaults);
+
+	if (bTemporarilyEnableDelta)
+	{
+		Ar.ArNoDelta = true;
+	}
 }
 
 uint32 UUserDefinedStruct::GetStructTypeHash(const void* Src) const
@@ -348,3 +413,22 @@ uint32 UUserDefinedStruct::GetUserDefinedStructTypeHash(const void* Src, const U
 	return ValueHash;
 }
 
+const uint8* UUserDefinedStruct::GetDefaultInstance() const
+{
+	ensure(DefaultStructInstance.IsValid() && DefaultStructInstance.GetStruct() == this);
+	return DefaultStructInstance.GetStructMemory();
+}
+
+void UUserDefinedStruct::AddReferencedObjects(UObject* InThis, FReferenceCollector& Collector)
+{
+	UUserDefinedStruct* This = CastChecked<UUserDefinedStruct>(InThis);
+
+	ensure(!This->DefaultStructInstance.IsValid() || This->DefaultStructInstance.GetStruct() == This);
+	uint8* StructData = This->DefaultStructInstance.GetStructMemory();
+	if (StructData)
+	{
+		This->SerializeBin(Collector.GetVerySlowReferenceCollectorArchive(), StructData);
+	}
+
+	Super::AddReferencedObjects(This, Collector);
+}
\ No newline at end of file
diff --git a/Engine/Source/Runtime/Engine/Private/WorldComposition.cpp b/Engine/Source/Runtime/Engine/Private/WorldComposition.cpp
--- a/Engine/Source/Runtime/Engine/Private/WorldComposition.cpp
+++ b/Engine/Source/Runtime/Engine/Private/WorldComposition.cpp
@@ -19,6 +19,7 @@
 #include "Engine/LevelStreaming.h"
 #include "Engine/LocalPlayer.h"
 #include "Engine/LevelStreamingKismet.h"
+#include "Engine/AssetManager.h"
 
 DEFINE_LOG_CATEGORY_STATIC(LogWorldComposition, Log, All);
 
@@ -129,6 +130,8 @@ struct FPackageNameAndLODIndex
 struct FWorldTilesGatherer
 	: public IPlatformFile::FDirectoryVisitor
 {
+	TArray<FString> MapFilesToConsider;
+
 	// List of tile long package names (non LOD)
 	TArray<FString> TilesCollection;
 	
@@ -137,31 +140,54 @@ struct FWorldTilesGatherer
 	
 	virtual bool Visit(const TCHAR* FilenameOrDirectory, bool bIsDirectory) override
 	{
-		FString FullPath = FilenameOrDirectory;
-	
 		// for all packages
-		if (!bIsDirectory && FPaths::GetExtension(FullPath, true) == FPackageName::GetMapPackageExtension())
+		if (!bIsDirectory)
+		{
+			FString FullPath = FilenameOrDirectory;
+
+			if (FPaths::GetExtension(FullPath, true) == FPackageName::GetMapPackageExtension())
+			{
+				MapFilesToConsider.Emplace(MoveTemp(FullPath));
+			}
+		}
+
+		return true;
+	}
+
+	void BuildTileCollection(const FString& WorldRootFilename)
+	{
+		FPlatformFileManager::Get().GetPlatformFile().IterateDirectoryRecursively(*WorldRootFilename, *this);
+
+		IAssetRegistry& AssetRegistry = UAssetManager::Get().GetAssetRegistry();
+		AssetRegistry.ScanFilesSynchronous(MapFilesToConsider);
+
+		for (const FString& FullPath : MapFilesToConsider)
 		{
 			FString TilePackageName = FPackageName::FilenameToLongPackageName(FullPath);
-			FPackageNameAndLODIndex PackageNameLOD = BreakToNameAndLODIndex(TilePackageName);
-						
-			if (PackageNameLOD.LODIndex != INDEX_NONE)
+			const FName TilePackagePath = *(TilePackageName + TEXT(".") + FPackageName::GetLongPackageAssetName(TilePackageName));
+
+			const FAssetData MapAssetData = AssetRegistry.GetAssetByObjectPath(TilePackagePath);
+
+			if (MapAssetData.IsValid() && !MapAssetData.IsRedirector())
 			{
-				if (PackageNameLOD.LODIndex == 0) 
-				{
-					// non-LOD tile
-					TilesCollection.Add(TilePackageName);
-				}
-				else
+				FPackageNameAndLODIndex PackageNameLOD = BreakToNameAndLODIndex(TilePackageName);
+
+				if (PackageNameLOD.LODIndex != INDEX_NONE)
 				{
-					// LOD tile
-					FString TileShortName = FPackageName::GetShortName(PackageNameLOD.PackageName);
-					TilesLODCollection.Add(TileShortName, PackageNameLOD);
+					if (PackageNameLOD.LODIndex == 0)
+					{
+						// non-LOD tile
+						TilesCollection.Emplace(MoveTemp(TilePackageName));
+					}
+					else
+					{
+						// LOD tile
+						FString TileShortName = FPackageName::GetShortName(PackageNameLOD.PackageName);
+						TilesLODCollection.Add(MoveTemp(TileShortName), MoveTemp(PackageNameLOD));
+					}
 				}
 			}
 		}
-	
-		return true;
 	}
 
 	FPackageNameAndLODIndex BreakToNameAndLODIndex(const FString& PackageName) const
@@ -216,7 +242,7 @@ void UWorldComposition::Rescan()
 	// Gather tiles packages from a specified folder
 	FWorldTilesGatherer Gatherer;
 	FString WorldRootFilename = FPackageName::LongPackageNameToFilename(WorldRoot);
-	FPlatformFileManager::Get().GetPlatformFile().IterateDirectoryRecursively(*WorldRootFilename, Gatherer);
+	Gatherer.BuildTileCollection(WorldRootFilename);
 
 	// Make sure we have persistent level name without PIE prefix
 	FString PersistentLevelPackageName = UWorld::StripPIEPrefixFromPackageName(OwningWorld->GetOutermost()->GetName(), OwningWorld->StreamingLevelsPrefix);
diff --git a/Engine/Source/Runtime/UMG/Private/WidgetBlueprintLibrary.cpp b/Engine/Source/Runtime/UMG/Private/WidgetBlueprintLibrary.cpp
--- a/Engine/Source/Runtime/UMG/Private/WidgetBlueprintLibrary.cpp
+++ b/Engine/Source/Runtime/UMG/Private/WidgetBlueprintLibrary.cpp
@@ -49,11 +49,6 @@ UUserWidget* UWidgetBlueprintLibrary::Create(UObject* WorldContextObject, TSubcl
 	{
 		UserWidget = CreateWidget<UUserWidget>(OwningPlayer, WidgetType);
 	}
-
-	if (UserWidget)
-	{
-		UserWidget->SetFlags(RF_StrongRefOnFrame);
-	}
 	return UserWidget;
 }
 
@@ -68,11 +63,6 @@ UDragDropOperation* UWidgetBlueprintLibrary::CreateDragDropOperation(TSubclassOf
 	{
 		DragDropOperation = NewObject<UDragDropOperation>();
 	}
-
-	if (DragDropOperation)
-	{
-		DragDropOperation->SetFlags(RF_StrongRefOnFrame);
-	}
 	return DragDropOperation;
 }
 
