Mirror 4094297 from Dev-Framework
Mark blueprint created properties as RF_LoadCompleted to prevent the linker from finding them and overwriting them

diff --git a/Engine/Source/Editor/KismetCompiler/Private/KismetCompiler.cpp b/Engine/Source/Editor/KismetCompiler/Private/KismetCompiler.cpp
--- a/Engine/Source/Editor/KismetCompiler/Private/KismetCompiler.cpp
+++ b/Engine/Source/Editor/KismetCompiler/Private/KismetCompiler.cpp
@@ -501,6 +501,12 @@ UProperty* FKismetCompilerContext::CreateVariable(const FName VarName, const FEd
 	UProperty* NewProperty = FKismetCompilerUtilities::CreatePropertyOnScope(NewClass, VarName, VarType, NewClass, CPF_None, Schema, MessageLog);
 	if (NewProperty != nullptr)
 	{
+		// This fixes a rare bug involving asynchronous loading of BPs in editor builds. The pattern was established
+		// in FKismetCompilerContext::CompileFunctions where we do this for the uber graph function. By setting
+		// the RF_LoadCompleted we prevent the linker from overwriting our regenerated property, although the
+		// circumstances under which this occurs are murky. More testing of BPs loading asynchronously in the editor
+		// needs to be added:
+		NewProperty->SetFlags(RF_LoadCompleted);
 		FKismetCompilerUtilities::LinkAddedProperty(NewClass, NewProperty);
 	}
 	else
