[NO NON-PII COMMIT MESSAGE]

diff --git a/Engine/Source/Editor/KismetCompiler/Private/KismetCompiler.cpp b/Engine/Source/Editor/KismetCompiler/Private/KismetCompiler.cpp
--- a/Engine/Source/Editor/KismetCompiler/Private/KismetCompiler.cpp
+++ b/Engine/Source/Editor/KismetCompiler/Private/KismetCompiler.cpp
@@ -2955,6 +2955,16 @@ void FKismetCompilerContext::ExpansionStep(UEdGraph* Graph, bool bAllowUbergraph
 		// If the pruning was called before expansion (and all collapsed graphs were saved), the isolated collapsed graphs would be unnecessarily validated.
 		PruneInner();
 
+		// First we need to expand knot nodes so any other expansions like AutoCreateRefTerm will have the correct pins hooked up
+		for (int32 NodeIndex = 0; NodeIndex < Graph->Nodes.Num(); ++NodeIndex)
+		{
+			UK2Node_Knot* KnotNode = Cast<UK2Node_Knot>(Graph->Nodes[NodeIndex]);
+			if (KnotNode)
+			{
+				KnotNode->ExpandNode(*this, Graph);
+			}
+		}
+
 		for (int32 NodeIndex = 0; NodeIndex < Graph->Nodes.Num(); ++NodeIndex)
 		{
 			UK2Node* Node = Cast<UK2Node>(Graph->Nodes[NodeIndex]);
