This change was needed due to changes to not create orphan pins on type changes, but is useful anyway as it could have happened with other function errors

diff --git a/Engine/Source/Editor/KismetCompiler/Private/KismetCompiler.cpp b/Engine/Source/Editor/KismetCompiler/Private/KismetCompiler.cpp
--- a/Engine/Source/Editor/KismetCompiler/Private/KismetCompiler.cpp
+++ b/Engine/Source/Editor/KismetCompiler/Private/KismetCompiler.cpp
@@ -3233,8 +3233,9 @@ void FKismetCompilerContext::CreateAndProcessUbergraph()
 				const UEdGraphNode* Node = ConsolidatedEventGraph->Nodes[ChildIndex];
 				const int32 SavedErrorCount = MessageLog.NumErrors;
 				UK2Node_Event* SrcEventNode = Cast<UK2Node_Event>(ConsolidatedEventGraph->Nodes[ChildIndex]);
-				if (bIsFullCompile || SrcEventNode)
+				if (bIsFullCompile)
 				{
+					// We only validate a full compile, we want to always make a function stub so we can display the errors for it later
 					ValidateNode(Node);
 				}
 
