Allow saving intermediate graphs when there is an error during compile

diff --git a/Engine/Source/Editor/KismetCompiler/Private/KismetCompiler.cpp b/Engine/Source/Editor/KismetCompiler/Private/KismetCompiler.cpp
--- a/Engine/Source/Editor/KismetCompiler/Private/KismetCompiler.cpp
+++ b/Engine/Source/Editor/KismetCompiler/Private/KismetCompiler.cpp
@@ -3991,23 +3991,6 @@ void FKismetCompilerContext::CompileFunctions(EInternalCompilerFlags InternalFla
 			}
 		}
 
-		// Save off intermediate build products if requested
-		if (CompileOptions.bSaveIntermediateProducts && !Blueprint->bIsRegeneratingOnLoad)
-		{
-			// Generate code for each function (done in a second pass to allow functions to reference each other)
-			for (int32 i = 0; i < FunctionList.Num(); ++i)
-			{
-				FKismetFunctionContext& ContextFunction = FunctionList[i];
-				if (FunctionList[i].SourceGraph != NULL)
-				{
-					// Record this graph as an intermediate product
-					ContextFunction.SourceGraph->Schema = UEdGraphSchema_K2::StaticClass();
-					Blueprint->IntermediateGeneratedGraphs.Add(ContextFunction.SourceGraph);
-					ContextFunction.SourceGraph->SetFlags(RF_Transient);
-				}
-			}
-		}
-
 		for (TFieldIterator<UMulticastDelegateProperty> PropertyIt(NewClass); PropertyIt; ++PropertyIt)
 		{
 			if(const UMulticastDelegateProperty* MCDelegateProp = *PropertyIt)
@@ -4033,6 +4016,23 @@ void FKismetCompilerContext::CompileFunctions(EInternalCompilerFlags InternalFla
 		}
 	}
 
+	// Save off intermediate build products if requested
+	if (bIsFullCompile && CompileOptions.bSaveIntermediateProducts && !Blueprint->bIsRegeneratingOnLoad)
+	{
+		// Generate code for each function (done in a second pass to allow functions to reference each other)
+		for (int32 i = 0; i < FunctionList.Num(); ++i)
+		{
+			FKismetFunctionContext& ContextFunction = FunctionList[i];
+			if (FunctionList[i].SourceGraph != NULL)
+			{
+				// Record this graph as an intermediate product
+				ContextFunction.SourceGraph->Schema = UEdGraphSchema_K2::StaticClass();
+				Blueprint->IntermediateGeneratedGraphs.Add(ContextFunction.SourceGraph);
+				ContextFunction.SourceGraph->SetFlags(RF_Transient);
+			}
+		}
+	}
+
 	// Late validation for Delegates.
 	{
 		TSet<UEdGraph*> AllGraphs;
