Interpreter Optimization 3 of N: Avoid creating UFunctions for fastpath animbp thunks

diff --git a/Engine/Source/Editor/KismetCompiler/Private/AnimBlueprintCompiler.cpp b/Engine/Source/Editor/KismetCompiler/Private/AnimBlueprintCompiler.cpp
--- a/Engine/Source/Editor/KismetCompiler/Private/AnimBlueprintCompiler.cpp
+++ b/Engine/Source/Editor/KismetCompiler/Private/AnimBlueprintCompiler.cpp
@@ -237,6 +237,11 @@ void FAnimBlueprintCompilerContext::CreateEvaluationHandlerStruct(UAnimGraphNode
 	check(Record.ServicedProperties.Num() > 0);
 	check(Record.NodeVariableProperty != NULL);
 
+	if (Record.IsFastPath())
+	{
+		return;
+	}
+
 	// Use the node GUID for a stable name across compiles
 	FString FunctionName = FString::Printf(TEXT("%s_%s_%s_%s"), *Record.EvaluationHandlerProperty->GetName(), *VisualAnimNode->GetOuter()->GetName(), *VisualAnimNode->GetClass()->GetName(), *VisualAnimNode->NodeGuid.ToString());
 	Record.HandlerFunctionName = FName(*FunctionName);
