Forgotten file from 4259893. Lost in shelf
#rnx

diff --git a/Engine/Source/Editor/KismetCompiler/Private/KismetCompiler.cpp b/Engine/Source/Editor/KismetCompiler/Private/KismetCompiler.cpp
--- a/Engine/Source/Editor/KismetCompiler/Private/KismetCompiler.cpp
+++ b/Engine/Source/Editor/KismetCompiler/Private/KismetCompiler.cpp
@@ -1737,6 +1737,7 @@ void FKismetCompilerContext::PrecompileFunction(FKismetFunctionContext& Context,
 		{
 			ensure(!NewClass->UberGraphFunction);
 			NewClass->UberGraphFunction = Context.Function;
+			NewClass->UberGraphFunction->FunctionFlags |= FUNC_UbergraphFunction;
 		}
 
 		// Register nets from function entry/exit nodes first, even for skeleton compiles (as they form the signature)
