Handle missing node types in AnimBlueprint GUID sanitizatin code

diff --git a/Engine/Source/Editor/KismetCompiler/Private/AnimBlueprintCompiler.cpp b/Engine/Source/Editor/KismetCompiler/Private/AnimBlueprintCompiler.cpp
--- a/Engine/Source/Editor/KismetCompiler/Private/AnimBlueprintCompiler.cpp
+++ b/Engine/Source/Editor/KismetCompiler/Private/AnimBlueprintCompiler.cpp
@@ -145,15 +145,18 @@ FAnimBlueprintCompilerContext::FAnimBlueprintCompilerContext(UAnimBlueprint* Sou
 
 					for (UEdGraphNode* Node : ChildGraph->Nodes)
 					{
-						if (NodeGuids.Contains(Node->NodeGuid))
+						if (Node)
 						{
-							bNodeGuidsRegenerated = true;
+							if (NodeGuids.Contains(Node->NodeGuid))
+							{
+								bNodeGuidsRegenerated = true;
 							
-							Node->CreateNewGuid(); // GUID is already being used, create a new one.
-						}
-						else
-						{
-							NodeGuids.Add(Node->NodeGuid);
+								Node->CreateNewGuid(); // GUID is already being used, create a new one.
+							}
+							else
+							{
+								NodeGuids.Add(Node->NodeGuid);
+							}
 						}
 					}
 				}
