Removed overzealous ensure because we support generating skeleton class in the middle of adding a new delegate
#rnx

diff --git a/Engine/Source/Editor/KismetCompiler/Private/KismetCompiler.cpp b/Engine/Source/Editor/KismetCompiler/Private/KismetCompiler.cpp
--- a/Engine/Source/Editor/KismetCompiler/Private/KismetCompiler.cpp
+++ b/Engine/Source/Editor/KismetCompiler/Private/KismetCompiler.cpp
@@ -656,7 +656,8 @@ void FKismetCompilerContext::CreateClassVariablesFromBlueprint()
 				if(UMulticastDelegateProperty* AsDelegate = Cast<UMulticastDelegateProperty>(NewProperty))
 				{
 					AsDelegate->SignatureFunction = FindField<UFunction>(NewClass, *(Variable.VarName.ToString() + HEADER_GENERATED_DELEGATE_SIGNATURE_SUFFIX));
-					ensure(AsDelegate->SignatureFunction);
+					// Skeleton compilation phase may run when the delegate has been created but the function has not:
+					ensureAlways(AsDelegate->SignatureFunction || !bIsFullCompile);
 				}
 			}
 
