Interpreter Optimization 9 of N: Mark ubergraph function as final so that calls to it don't need to search the target object

diff --git a/Engine/Source/Editor/KismetCompiler/Private/KismetCompiler.cpp b/Engine/Source/Editor/KismetCompiler/Private/KismetCompiler.cpp
--- a/Engine/Source/Editor/KismetCompiler/Private/KismetCompiler.cpp
+++ b/Engine/Source/Editor/KismetCompiler/Private/KismetCompiler.cpp
@@ -1760,6 +1760,7 @@ void FKismetCompilerContext::PrecompileFunction(FKismetFunctionContext& Context,
 			ensure(!NewClass->UberGraphFunction);
 			NewClass->UberGraphFunction = Context.Function;
 			NewClass->UberGraphFunction->FunctionFlags |= FUNC_UbergraphFunction;
+			NewClass->UberGraphFunction->FunctionFlags |= FUNC_Final;
 		}
 
 		// Register nets from function entry/exit nodes first, even for skeleton compiles (as they form the signature)
