Non-container types, and native created types already worked

diff --git a/Engine/Source/Editor/KismetCompiler/Private/KismetCompilerMisc.cpp b/Engine/Source/Editor/KismetCompiler/Private/KismetCompilerMisc.cpp
--- a/Engine/Source/Editor/KismetCompiler/Private/KismetCompilerMisc.cpp
+++ b/Engine/Source/Editor/KismetCompiler/Private/KismetCompilerMisc.cpp
@@ -1162,7 +1162,7 @@ UProperty* FKismetCompilerUtilities::CreatePropertyOnScope(UStruct* Scope, const
 		NewProperty = CreatePrimitiveProperty(PropertyScope, ValidatedPropertyName, Type.PinCategory, Type.PinSubCategory, Type.PinSubCategoryObject.Get(), SelfClass, Type.bIsWeakPointer, Schema, MessageLog);
 	}
 
-	if (NewContainerProperty && NewProperty && NewProperty->HasAnyPropertyFlags(CPF_ContainsInstancedReference))
+	if (NewContainerProperty && NewProperty && NewProperty->HasAnyPropertyFlags(CPF_ContainsInstancedReference | CPF_InstancedReference))
 	{
 		NewContainerProperty->SetPropertyFlags(CPF_ContainsInstancedReference);
 	}
@@ -1196,9 +1196,9 @@ UProperty* FKismetCompilerUtilities::CreatePropertyOnScope(UStruct* Scope, const
 			}
 			else
 			{
-				if (NewMapProperty->ValueProp && NewMapProperty->ValueProp->HasAnyPropertyFlags(CPF_ContainsInstancedReference))
+				if (NewMapProperty->ValueProp && NewMapProperty->ValueProp->HasAnyPropertyFlags(CPF_ContainsInstancedReference | CPF_InstancedReference))
 				{
-					NewMapProperty->ValueProp->SetPropertyFlags(CPF_ContainsInstancedReference);
+					NewContainerProperty->SetPropertyFlags(CPF_ContainsInstancedReference);
 				}
 
 				NewProperty = NewMapProperty;
