Undo //UE4/Release-4.22/Engine/Source/Editor/KismetCompiler/Private/... changelist 6317679

diff --git a/Engine/Source/Editor/KismetCompiler/Private/KismetCompilerMisc.cpp b/Engine/Source/Editor/KismetCompiler/Private/KismetCompilerMisc.cpp
--- a/Engine/Source/Editor/KismetCompiler/Private/KismetCompilerMisc.cpp
+++ b/Engine/Source/Editor/KismetCompiler/Private/KismetCompilerMisc.cpp
@@ -1162,7 +1162,7 @@ UProperty* FKismetCompilerUtilities::CreatePropertyOnScope(UStruct* Scope, const
 		NewProperty = CreatePrimitiveProperty(PropertyScope, ValidatedPropertyName, Type.PinCategory, Type.PinSubCategory, Type.PinSubCategoryObject.Get(), SelfClass, Type.bIsWeakPointer, Schema, MessageLog);
 	}
 
-	if (NewContainerProperty && NewProperty && NewProperty->HasAnyPropertyFlags(CPF_ContainsInstancedReference | CPF_InstancedReference))
+	if (NewContainerProperty && NewProperty && NewProperty->HasAnyPropertyFlags(CPF_ContainsInstancedReference))
 	{
 		NewContainerProperty->SetPropertyFlags(CPF_ContainsInstancedReference);
 	}
@@ -1196,9 +1196,9 @@ UProperty* FKismetCompilerUtilities::CreatePropertyOnScope(UStruct* Scope, const
 			}
 			else
 			{
-				if (NewMapProperty->ValueProp && NewMapProperty->ValueProp->HasAnyPropertyFlags(CPF_ContainsInstancedReference | CPF_InstancedReference))
+				if (NewMapProperty->ValueProp && NewMapProperty->ValueProp->HasAnyPropertyFlags(CPF_ContainsInstancedReference))
 				{
-					NewContainerProperty->SetPropertyFlags(CPF_ContainsInstancedReference);
+					NewMapProperty->ValueProp->SetPropertyFlags(CPF_ContainsInstancedReference);
 				}
 
 				NewProperty = NewMapProperty;
