Sub Anim Blueprints: Unable to watch values that are passed to Animation Layers
#fix Make sure we register debug property mapping for sub input node properties, they used to be incorrectly mapped to gc-ed objects

diff --git a/Engine/Source/Editor/KismetCompiler/Private/AnimBlueprintCompiler.cpp b/Engine/Source/Editor/KismetCompiler/Private/AnimBlueprintCompiler.cpp
--- a/Engine/Source/Editor/KismetCompiler/Private/AnimBlueprintCompiler.cpp
+++ b/Engine/Source/Editor/KismetCompiler/Private/AnimBlueprintCompiler.cpp
@@ -2137,6 +2137,13 @@ void FAnimBlueprintCompilerContext::ProcessSubInput(UAnimGraphNode_SubInput* InS
 						VariableGetNode->SetFromProperty(NewSubInputProperty, true);
 						VariableGetNode->AllocateDefaultPins();
 
+						// Add pin to generated variable association, used for pin watching
+						UEdGraphPin* TrueSourcePin = MessageLog.FindSourcePin(Pin);
+						if (TrueSourcePin)
+						{
+							NewClass->GetDebugData().RegisterClassPropertyAssociation(TrueSourcePin, NewSubInputProperty);
+						}
+
 						// link up to new node - note that this is not a FindPinChecked because if an interface changes without the
 						// implementing class being loaded, then its graphs will not be conformed until AFTER the skeleton class
 						// has been compiled, so the variable cannot be created. This also doesnt matter, as there wont be anything connected
