Improve consistency of function flag checking in code generator
#rnx

diff --git a/Engine/Source/Editor/KismetCompiler/Private/KismetCompilerVMBackend.cpp b/Engine/Source/Editor/KismetCompiler/Private/KismetCompilerVMBackend.cpp
--- a/Engine/Source/Editor/KismetCompiler/Private/KismetCompilerVMBackend.cpp
+++ b/Engine/Source/Editor/KismetCompiler/Private/KismetCompilerVMBackend.cpp
@@ -1199,7 +1199,7 @@ public:
 			&& !UEdGraphSchema_K2::HasWildcardParams(FunctionToCall);
 
 		const bool bLocalScriptFunction = 
-			!FunctionToCall->HasAnyFunctionFlags(FUNC_Native|FUNC_NetFuncFlags);
+			!FunctionToCall->HasAnyFunctionFlags(FUNC_Native|FUNC_NetFuncFlags|FUNC_BlueprintAuthorityOnly|FUNC_BlueprintCosmetic|FUNC_NetRequest|FUNC_NetResponse);
 
 		// Handle the function calling context if needed
 		FContextEmitter CallContextWriter(*this);
