Remove dangerous, unused function

diff --git a/Engine/Source/Editor/KismetCompiler/Private/KismetCompilerMisc.cpp b/Engine/Source/Editor/KismetCompiler/Private/KismetCompilerMisc.cpp
--- a/Engine/Source/Editor/KismetCompiler/Private/KismetCompilerMisc.cpp
+++ b/Engine/Source/Editor/KismetCompiler/Private/KismetCompilerMisc.cpp
@@ -512,33 +512,6 @@ void FKismetCompilerUtilities::RemoveObjectRedirectorIfPresent(UObject* Package,
 	}
 }
 
-void FKismetCompilerUtilities::EnsureFreeNameForNewClass(UClass* ClassToConsign, FString& ClassName, UBlueprint* Blueprint)
-{
-	check(Blueprint);
-
-	UObject* OwnerOutermost = Blueprint->GetOutermost();
-
-	// Try to find a class with the name we want to use in the scope
-	UClass* AnyClassWithGoodName = (UClass*)StaticFindObject(UClass::StaticClass(), OwnerOutermost, *ClassName, false);
-	if (AnyClassWithGoodName == ClassToConsign)
-	{
-		// Ignore it if it's the class we're already consigning anyway
-		AnyClassWithGoodName = NULL;
-	}
-
-	if( ClassToConsign )
-	{
-		ConsignToOblivion(ClassToConsign, Blueprint->bIsRegeneratingOnLoad);
-	}
-
-	// Consign the class with the name we want to use
-	if( AnyClassWithGoodName )
-	{
-		ConsignToOblivion(AnyClassWithGoodName, Blueprint->bIsRegeneratingOnLoad);
-	}
-}
-
-
 /** Finds a property by name, starting in the specified scope; Validates property type and returns NULL along with emitting an error if there is a mismatch. */
 UProperty* FKismetCompilerUtilities::FindPropertyInScope(UStruct* Scope, UEdGraphPin* Pin, FCompilerResultsLog& MessageLog, const UEdGraphSchema_K2* Schema, UClass* SelfClass)
 {
diff --git a/Engine/Source/Editor/KismetCompiler/Public/KismetCompilerMisc.h b/Engine/Source/Editor/KismetCompiler/Public/KismetCompilerMisc.h
--- a/Engine/Source/Editor/KismetCompiler/Public/KismetCompilerMisc.h
+++ b/Engine/Source/Editor/KismetCompiler/Public/KismetCompilerMisc.h
@@ -37,16 +37,6 @@ public:
 	 */
 	static void InvalidatePropertyExport(UProperty* PropertyToInvalidate);
 
-	/**
-	 * Finds any class with the specified class name, and consigns it to oblivion, along with the specified class to consign.
-	 * This should ensure that the specified name is free for use
-	 *
-	 * @param ClassToConsign	Class to consign to oblivion no matter what
-	 * @param ClassName			Name that we want to ensure isn't used
-	 * @param Blueprint			The blueprint that is in charge of these classes, used for scoping
-	 */
-	static void EnsureFreeNameForNewClass(UClass* ClassToConsign, FString& ClassName, UBlueprint* Blueprint);
-
 	/**
 	 * Tests to see if a pin is schema compatible with a property.
 	 *
