Fix potential CDO validation failures at cook time during the C++ code generation pass when Blueprint nativization is enabled.

- Mirrored from 4371073.

diff --git a/Engine/Source/Editor/KismetCompiler/Private/KismetCompiler.cpp b/Engine/Source/Editor/KismetCompiler/Private/KismetCompiler.cpp
--- a/Engine/Source/Editor/KismetCompiler/Private/KismetCompiler.cpp
+++ b/Engine/Source/Editor/KismetCompiler/Private/KismetCompiler.cpp
@@ -4016,7 +4016,7 @@ void FKismetCompilerContext::CompileFunctions(EInternalCompilerFlags InternalFla
 	FKismetCompilerVMBackend Backend_VM(Blueprint, Schema, *this);
 
 	// Determine whether or not to skip generated class validation. This requires CDO value propagation to occur first.
-	bool bSkipGeneratedClassValidation = !bPropagateValuesToCDO;
+	bool bSkipGeneratedClassValidation = !bPropagateValuesToCDO || CompileOptions.CompileType == EKismetCompileType::Cpp;
 
 	if( bGenerateLocals )
 	{
