Fix data loss regression on nativized builds.

diff --git a/Engine/Source/Editor/KismetCompiler/Private/KismetCompiler.cpp b/Engine/Source/Editor/KismetCompiler/Private/KismetCompiler.cpp
--- a/Engine/Source/Editor/KismetCompiler/Private/KismetCompiler.cpp
+++ b/Engine/Source/Editor/KismetCompiler/Private/KismetCompiler.cpp
@@ -4181,7 +4181,7 @@ void FKismetCompilerContext::CompileFunctions(EInternalCompilerFlags InternalFla
 				}
 
 				// >>> Backwards Compatibility: Propagate data from the skel CDO to the gen CDO if we haven't already done so for this blueprint
-				if( !bIsSkeletonOnly && !Blueprint->IsGeneratedClassAuthoritative() )
+				if( !bIsSkeletonOnly && !Blueprint->IsGeneratedClassAuthoritative() && CompileOptions.CompileType != EKismetCompileType::Cpp )
 				{
 					UEditorEngine::FCopyPropertiesForUnrelatedObjectsParams CopyDetails;
 					CopyDetails.bAggressiveDefaultSubobjectReplacement = false;
