(4.21.2) Don't try and set metadata on a pin with none for a name.
#rnx

diff --git a/Engine/Source/Editor/KismetCompiler/Private/KismetCompiler.cpp b/Engine/Source/Editor/KismetCompiler/Private/KismetCompiler.cpp
--- a/Engine/Source/Editor/KismetCompiler/Private/KismetCompiler.cpp
+++ b/Engine/Source/Editor/KismetCompiler/Private/KismetCompiler.cpp
@@ -2256,8 +2256,9 @@ void FKismetCompilerContext::SetDefaultInputValueMetaData(UFunction* Function, c
 	for (const TSharedPtr<FUserPinInfo>& InputDataPtr : InputData)
 	{
 		if ( InputDataPtr.IsValid() &&
+			!InputDataPtr->PinName.IsNone() &&
+			(InputDataPtr->PinName != UEdGraphSchema_K2::PN_Self) &&
 			(InputDataPtr->PinType.PinCategory != UEdGraphSchema_K2::PC_Exec) && 
-			(InputDataPtr->PinName != UEdGraphSchema_K2::PN_Self) && 
 			(InputDataPtr->PinType.PinCategory != UEdGraphSchema_K2::PC_Object) &&
 			(InputDataPtr->PinType.PinCategory != UEdGraphSchema_K2::PC_Class) &&
 			(InputDataPtr->PinType.PinCategory != UEdGraphSchema_K2::PC_Interface) )
