Make sure we don't cause a name collision when moving old class subobject aside

diff --git a/Engine/Source/Editor/KismetCompiler/Private/KismetCompiler.cpp b/Engine/Source/Editor/KismetCompiler/Private/KismetCompiler.cpp
--- a/Engine/Source/Editor/KismetCompiler/Private/KismetCompiler.cpp
+++ b/Engine/Source/Editor/KismetCompiler/Private/KismetCompiler.cpp
@@ -277,7 +277,8 @@ void FKismetCompilerContext::CleanAndSanitizeClass(UBlueprintGeneratedClass* Cla
 	for( auto SubObjIt = ClassSubObjects.CreateIterator(); SubObjIt; ++SubObjIt )
 	{
 		UObject* CurrSubObj = *SubObjIt;
-		CurrSubObj->Rename(*CurrSubObj->GetName(), TransientClass, RenFlags);
+		FName NewSubobjectName = MakeUniqueObjectName(TransientClass, CurrSubObj->GetClass(), CurrSubObj->GetFName());
+		CurrSubObj->Rename(*NewSubobjectName.ToString(), TransientClass, RenFlags);
 		if( UProperty* Prop = Cast<UProperty>(CurrSubObj) )
 		{
 			FKismetCompilerUtilities::InvalidatePropertyExport(Prop);
