Fixed ensure when compiling a blueprint interface after changing a function signature

diff --git a/Engine/Source/Editor/KismetCompiler/Private/KismetCompiler.cpp b/Engine/Source/Editor/KismetCompiler/Private/KismetCompiler.cpp
--- a/Engine/Source/Editor/KismetCompiler/Private/KismetCompiler.cpp
+++ b/Engine/Source/Editor/KismetCompiler/Private/KismetCompiler.cpp
@@ -3987,14 +3987,21 @@ void FKismetCompilerContext::CompileClassLayout(EInternalCompilerFlags InternalF
 		}
 	}
 
-	// Make sure that this blueprint is up-to-date with regards to its parent functions
-	FBlueprintEditorUtils::ConformCallsToParentFunctions(Blueprint);
+	{
+		// the following calls may mark the blueprint as dirty, but we know that these operations just cleaned up the BP 
+		// so dependencies can still be considered 'up to date'
+		TGuardValue<bool> LockDependenciesUpToDate(Blueprint->bCachedDependenciesUpToDate, Blueprint->bCachedDependenciesUpToDate);
+
+		// Make sure that this blueprint is up-to-date with regards to its parent functions
+		FBlueprintEditorUtils::ConformCallsToParentFunctions(Blueprint);
 
-	// Conform implemented events here, to ensure we generate custom events if necessary after reparenting
-	FBlueprintEditorUtils::ConformImplementedEvents(Blueprint);
+		// Conform implemented events here, to ensure we generate custom events if necessary after reparenting
+		FBlueprintEditorUtils::ConformImplementedEvents(Blueprint);
+
+		// Conform implemented interfaces here, to ensure we generate all functions required by the interface as stubs
+		FBlueprintEditorUtils::ConformImplementedInterfaces(Blueprint);
+	}
 
-	// Conform implemented interfaces here, to ensure we generate all functions required by the interface as stubs
-	FBlueprintEditorUtils::ConformImplementedInterfaces(Blueprint);
 
 	// Run thru the class defined variables first, get them registered
 	CreateClassVariablesFromBlueprint();
