Fixed compiled Blueprint output terms to use a format that will utilize the name table more efficiently.

diff --git a/Engine/Source/Editor/BlueprintGraph/Private/CallFunctionHandler.cpp b/Engine/Source/Editor/BlueprintGraph/Private/CallFunctionHandler.cpp
--- a/Engine/Source/Editor/BlueprintGraph/Private/CallFunctionHandler.cpp
+++ b/Engine/Source/Editor/BlueprintGraph/Private/CallFunctionHandler.cpp
@@ -566,7 +566,7 @@ void FKCHandler_CallFunction::RegisterNets(FKismetFunctionContext& Context, UEdG
 			((Pin->PinType.PinCategory == UEdGraphSchema_K2::PC_Object) && (Pin->LinkedTo[0]->PinType.PinCategory == UEdGraphSchema_K2::PC_Interface)))
 		{
 			FBPTerminal* InterfaceTerm = Context.CreateLocalTerminal();
-			InterfaceTerm->CopyFromPin(Pin, Context.NetNameMap->MakeValidName(Pin) + TEXT("_CastInput"));
+			InterfaceTerm->CopyFromPin(Pin, Context.NetNameMap->MakeValidName(Pin, TEXT("CastInput")));
 			InterfaceTerm->Source = Node;
 
 			InterfaceTermMap.Add(Pin, InterfaceTerm);
diff --git a/Engine/Source/Editor/BlueprintGraph/Private/DynamicCastHandler.cpp b/Engine/Source/Editor/BlueprintGraph/Private/DynamicCastHandler.cpp
--- a/Engine/Source/Editor/BlueprintGraph/Private/DynamicCastHandler.cpp
+++ b/Engine/Source/Editor/BlueprintGraph/Private/DynamicCastHandler.cpp
@@ -25,7 +25,7 @@ void FKCHandler_DynamicCast::RegisterNets(FKismetFunctionContext& Context, UEdGr
 			FBPTerminal* BoolTerm = Context.CreateLocalTerminal();
 			BoolTerm->Type.PinCategory = UEdGraphSchema_K2::PC_Boolean;
 			BoolTerm->Source = Node;
-			BoolTerm->Name = Context.NetNameMap->MakeValidName(Node) + TEXT("_CastSuccess");
+			BoolTerm->Name = Context.NetNameMap->MakeValidName(Node, TEXT("CastSuccess"));
 			BoolTermMap.Add(Node, BoolTerm);
 		}
 	}
diff --git a/Engine/Source/Editor/BlueprintGraph/Private/K2Node_DelegateSet.cpp b/Engine/Source/Editor/BlueprintGraph/Private/K2Node_DelegateSet.cpp
--- a/Engine/Source/Editor/BlueprintGraph/Private/K2Node_DelegateSet.cpp
+++ b/Engine/Source/Editor/BlueprintGraph/Private/K2Node_DelegateSet.cpp
@@ -38,7 +38,7 @@ public:
 			DelegateTerm->Type.PinCategory = UEdGraphSchema_K2::PC_Delegate;
 			FMemberReference::FillSimpleMemberReference<UFunction>(DelegateNode->GetDelegateSignature(), DelegateTerm->Type.PinSubCategoryMemberReference);
 			DelegateTerm->Source = Node;
-			DelegateTerm->Name = Context.NetNameMap->MakeValidName(Node) + TEXT("_TempBindingDelegate");
+			DelegateTerm->Name = Context.NetNameMap->MakeValidName(Node, TEXT("TempBindingDelegate"));
 			LocalDelegateMap.Add(Node, DelegateTerm);
 
 			// The only net we need to register for this node is the delegate's target (self) pin, since the others are expanded to their own event node
diff --git a/Engine/Source/Editor/BlueprintGraph/Private/K2Node_Select.cpp b/Engine/Source/Editor/BlueprintGraph/Private/K2Node_Select.cpp
--- a/Engine/Source/Editor/BlueprintGraph/Private/K2Node_Select.cpp
+++ b/Engine/Source/Editor/BlueprintGraph/Private/K2Node_Select.cpp
@@ -75,7 +75,7 @@ public:
 				return;
 			}
 
-			FString DefaultTermName = Context.NetNameMap->MakeValidName(Node) + TEXT("_Default");
+			FString DefaultTermName = Context.NetNameMap->MakeValidName(Node, TEXT("Default"));
 			FBPTerminal* DefaultTerm = Context.CreateLocalTerminalFromPinAutoChooseScope(OptionPins[0], DefaultTermName);
 			DefaultTermMap.Add(Node, DefaultTerm);
 		}
diff --git a/Engine/Source/Editor/BlueprintGraph/Private/K2Node_Switch.cpp b/Engine/Source/Editor/BlueprintGraph/Private/K2Node_Switch.cpp
--- a/Engine/Source/Editor/BlueprintGraph/Private/K2Node_Switch.cpp
+++ b/Engine/Source/Editor/BlueprintGraph/Private/K2Node_Switch.cpp
@@ -44,7 +44,7 @@ public:
 		FBPTerminal* BoolTerm = Context.CreateLocalTerminal();
 		BoolTerm->Type.PinCategory = UEdGraphSchema_K2::PC_Boolean;
 		BoolTerm->Source = Node;
-		BoolTerm->Name = Context.NetNameMap->MakeValidName(Node) + TEXT("_CmpSuccess");
+		BoolTerm->Name = Context.NetNameMap->MakeValidName(Node, TEXT("CmpSuccess"));
 		BoolTermMap.Add(Node, BoolTerm);
 	}
 
diff --git a/Engine/Source/Editor/KismetCompiler/Public/KismetCompilerMisc.h b/Engine/Source/Editor/KismetCompiler/Public/KismetCompilerMisc.h
--- a/Engine/Source/Editor/KismetCompiler/Public/KismetCompilerMisc.h
+++ b/Engine/Source/Editor/KismetCompiler/Public/KismetCompilerMisc.h
@@ -206,11 +206,11 @@ public:
 struct FNetNameMapping
 {
 public:
-	// Come up with a valid, unique (within the scope of NetNameMap) name based on an existing Net object.
+	// Come up with a valid, unique (within the scope of NetNameMap) name based on an existing Net object and (optional) context.
 	// The resulting name is stable across multiple calls if given the same pointer.
-	FString MakeValidName(const UEdGraphNode* Net) { return MakeValidNameImpl(Net); }
-	FString MakeValidName(const UEdGraphPin* Net) { return MakeValidNameImpl(Net); }
-	FString MakeValidName(const UAnimGraphNode_Base* Net){ return MakeValidNameImpl(Net); }
+	FString MakeValidName(const UEdGraphNode* Net, const FString& Context = TEXT("")) { return MakeValidNameImpl(Net, Context); }
+	FString MakeValidName(const UEdGraphPin* Net, const FString& Context = TEXT("")) { return MakeValidNameImpl(Net, Context); }
+	FString MakeValidName(const UAnimGraphNode_Base* Net, const FString& Context = TEXT("")) { return MakeValidNameImpl(Net, Context); }
 
 private:
 	KISMETCOMPILER_API static FString MakeBaseName(const UEdGraphNode* Net);
@@ -218,7 +218,7 @@ private:
 	KISMETCOMPILER_API static FString MakeBaseName(const UAnimGraphNode_Base* Net);
 
 	template< typename NetType >
-	FString MakeValidNameImpl(NetType Net)
+	FString MakeValidNameImpl(NetType Net, const FString& Context)
 	{
 		// Check to see if this net was already used to generate a name
 		if (FString* Result = NetToName.Find(Net))
@@ -227,7 +227,14 @@ private:
 		}
 		else
 		{
-			FString NetName = GetUniqueName(MakeBaseName(Net));
+			FString BaseName = MakeBaseName(Net);
+			if (!Context.IsEmpty())
+			{
+				BaseName += FString::Printf(TEXT("_%s"), *Context);
+			}
+
+			FString NetName = GetUniqueName(MoveTemp(BaseName));
+
 			NetToName.Add(Net, NetName);
 			NameToNet.Add(NetName, Net);
 			return NetName;
@@ -246,7 +253,7 @@ private:
 		{
 			++Postfix;
 			// Add an integer to the base name and check if it's free:
-			NewNetName = NetName + FString::FromInt(Postfix);
+			NewNetName = NetName + TEXT("_") + FString::FromInt(Postfix);
 			ExistingNet = NameToNet.Find(NewNetName);
 		}
 		return NewNetName;
