Fix applying a layer overlay when the applied class implements more than one layer

Also makes sure that sub-input nodes are always preserved even if not connected to the root, to avoid suprious warnings when dynamic linking.

#rnx

diff --git a/Engine/Source/Editor/KismetCompiler/Private/AnimBlueprintCompiler.cpp b/Engine/Source/Editor/KismetCompiler/Private/AnimBlueprintCompiler.cpp
--- a/Engine/Source/Editor/KismetCompiler/Private/AnimBlueprintCompiler.cpp
+++ b/Engine/Source/Editor/KismetCompiler/Private/AnimBlueprintCompiler.cpp
@@ -913,63 +913,34 @@ int32 FAnimBlueprintCompilerContext::GetAllocationIndexOfNode(UAnimGraphNode_Bas
 
 void FAnimBlueprintCompilerContext::PruneIsolatedAnimationNodes(const TArray<UAnimGraphNode_Base*>& RootSet, TArray<UAnimGraphNode_Base*>& GraphNodes)
 {
-	struct FNodeVisitorDownWires
+	struct FNodeVisitorDownPoseWires
 	{
 		TSet<UEdGraphNode*> VisitedNodes;
-		TSet<UEdGraphNode*> VisitedNodesViaPropertyPins;
 		const UAnimationGraphSchema* Schema;
 
-		FNodeVisitorDownWires()
+		FNodeVisitorDownPoseWires()
 		{
 			Schema = GetDefault<UAnimationGraphSchema>();
 		}
 
-		void TraverseWires(UEdGraphNode* Node, bool bProperty = false)
+		void TraverseNodes(UEdGraphNode* Node)
 		{
-			if(bProperty)
-			{
-				VisitedNodesViaPropertyPins.Add(Node);
-
-				// We only care about sub-inputs here
-				if(Node->IsA<UAnimGraphNode_SubInput>())
-				{
-					VisitedNodes.Add(Node);
-				}
-			}
-			else
-			{
-				VisitedNodes.Add(Node);
-			}
+			VisitedNodes.Add(Node);
 
 			// Follow every exec output pin
 			for (int32 i = 0; i < Node->Pins.Num(); ++i)
 			{
 				UEdGraphPin* MyPin = Node->Pins[i];
 
-				if (MyPin->Direction == EGPD_Input)
+				if ((MyPin->Direction == EGPD_Input) && (Schema->IsPosePin(MyPin->PinType)))
 				{
-					if(Schema->IsPosePin(MyPin->PinType))
-					{
-						for (int32 j = 0; j < MyPin->LinkedTo.Num(); ++j)
-						{
-							UEdGraphPin* OtherPin = MyPin->LinkedTo[j];
-							UEdGraphNode* OtherNode = OtherPin->GetOwningNode();
-							if (!VisitedNodes.Contains(OtherNode))
-							{
-								TraverseWires(OtherNode, false);
-							}
-						}
-					}
-					else
+					for (int32 j = 0; j < MyPin->LinkedTo.Num(); ++j)
 					{
-						for (int32 j = 0; j < MyPin->LinkedTo.Num(); ++j)
+						UEdGraphPin* OtherPin = MyPin->LinkedTo[j];
+						UEdGraphNode* OtherNode = OtherPin->GetOwningNode();
+						if (!VisitedNodes.Contains(OtherNode))
 						{
-							UEdGraphPin* OtherPin = MyPin->LinkedTo[j];
-							UEdGraphNode* OtherNode = OtherPin->GetOwningNode();
-							if (!VisitedNodesViaPropertyPins.Contains(OtherNode) && !VisitedNodes.Contains(OtherNode))
-							{
-								TraverseWires(OtherNode, true);
-							}
+							TraverseNodes(OtherNode);
 						}
 					}
 				}
@@ -977,19 +948,21 @@ void FAnimBlueprintCompilerContext::PruneIsolatedAnimationNodes(const TArray<UAn
 		}
 	};
 
-	// Prune the nodes that aren't reachable via an animation pose link (or linked to a sub input node)
-	FNodeVisitorDownWires Visitor;
+	// Prune the nodes that aren't reachable via an animation pose link
+	FNodeVisitorDownPoseWires Visitor;
 
 	for (auto RootIt = RootSet.CreateConstIterator(); RootIt; ++RootIt)
 	{
 		UAnimGraphNode_Base* RootNode = *RootIt;
-		Visitor.TraverseWires(RootNode);
+		Visitor.TraverseNodes(RootNode);
 	}
 
 	for (int32 NodeIndex = 0; NodeIndex < GraphNodes.Num(); ++NodeIndex)
 	{
 		UAnimGraphNode_Base* Node = GraphNodes[NodeIndex];
-		if (!Visitor.VisitedNodes.Contains(Node) && !IsNodePure(Node))
+
+		// We cant prune sub-inputs as even if they are not linked to the root, they are needed for the dynamic link phase at runtime
+		if (!Visitor.VisitedNodes.Contains(Node) && !IsNodePure(Node) && !Node->IsA<UAnimGraphNode_SubInput>())
 		{
 			Node->BreakAllNodeLinks();
 			GraphNodes.RemoveAtSwap(NodeIndex);
diff --git a/Engine/Source/Runtime/Engine/Private/Animation/AnimInstance.cpp b/Engine/Source/Runtime/Engine/Private/Animation/AnimInstance.cpp
--- a/Engine/Source/Runtime/Engine/Private/Animation/AnimInstance.cpp
+++ b/Engine/Source/Runtime/Engine/Private/Animation/AnimInstance.cpp
@@ -2560,10 +2560,13 @@ void UAnimInstance::SetLayerOverlay(TSubclassOf<UAnimInstance> InClass)
 					// Init after we link in the new graph segments, so propagation happens correctly
 					NewSubInstance->InitializeAnimation();
 
-					// Initialize the correct part of the sub instance
-					if(LayerPair.Value[0]->LinkedRoot)
+					// Initialize the correct parts of the sub instance
+					for(FAnimNode_Layer* LayerNode : LayerPair.Value)
 					{
-						NewSubInstance->GetProxyOnAnyThread<FAnimInstanceProxy>().InitializeRootNode_WithRoot(LayerPair.Value[0]->LinkedRoot);
+						if(LayerNode->LinkedRoot)
+						{
+							NewSubInstance->GetProxyOnAnyThread<FAnimInstanceProxy>().InitializeRootNode_WithRoot(LayerNode->LinkedRoot);
+						}
 					}
 
 					MeshComp->SubInstances.Add(NewSubInstance);
