Wider use of EX_CallMath

diff --git a/Engine/Source/Editor/BlueprintGraph/Classes/EdGraphSchema_K2.h b/Engine/Source/Editor/BlueprintGraph/Classes/EdGraphSchema_K2.h
--- a/Engine/Source/Editor/BlueprintGraph/Classes/EdGraphSchema_K2.h
+++ b/Engine/Source/Editor/BlueprintGraph/Classes/EdGraphSchema_K2.h
@@ -733,6 +733,9 @@ public:
 	 */
 	bool ConvertPropertyToPinType(const UProperty* Property, /*out*/ FEdGraphPinType& TypeOut) const;
 
+	/** Returns true if the function has wildcard parameters, e.g. uses runtime type information that may require safe failure handling */
+	static bool HasWildcardParams(const UFunction* Function);
+
 	/** Determines if the specified param property is intended to be used as a wildcard (for custom thunk functions, like in our array library, etc.)*/
 	static bool IsWildcardProperty(const UProperty* ParamProperty);
 
diff --git a/Engine/Source/Editor/BlueprintGraph/Private/EdGraphSchema_K2.cpp b/Engine/Source/Editor/BlueprintGraph/Private/EdGraphSchema_K2.cpp
--- a/Engine/Source/Editor/BlueprintGraph/Private/EdGraphSchema_K2.cpp
+++ b/Engine/Source/Editor/BlueprintGraph/Private/EdGraphSchema_K2.cpp
@@ -3583,6 +3583,21 @@ bool UEdGraphSchema_K2::ConvertPropertyToPinType(const UProperty* Property, /*ou
 	return true;
 }
 
+bool UEdGraphSchema_K2::HasWildcardParams(const UFunction* Function)
+{
+	bool bResult = false;
+	for (TFieldIterator<const UProperty> PropIt(Function); PropIt && (PropIt->PropertyFlags & CPF_Parm) && !bResult; ++PropIt)
+	{
+		const UProperty* FuncParamProperty = *PropIt;
+
+		if (IsWildcardProperty(FuncParamProperty))
+		{
+			bResult = true;
+		}
+	}
+	return bResult;
+}
+
 bool UEdGraphSchema_K2::IsWildcardProperty(const UProperty* Property)
 {
 	UFunction* Function = Cast<UFunction>(Property->GetOuter());
diff --git a/Engine/Source/Editor/KismetCompiler/Private/KismetCompilerVMBackend.cpp b/Engine/Source/Editor/KismetCompiler/Private/KismetCompilerVMBackend.cpp
--- a/Engine/Source/Editor/KismetCompiler/Private/KismetCompilerVMBackend.cpp
+++ b/Engine/Source/Editor/KismetCompiler/Private/KismetCompilerVMBackend.cpp
@@ -1192,7 +1192,11 @@ public:
 			&& FunctionToCall->HasAllFunctionFlags(FUNC_Static | FUNC_BlueprintPure | FUNC_Final | FUNC_Native)
 			&& !FunctionToCall->HasAnyFunctionFlags(FUNC_BlueprintAuthorityOnly | FUNC_BlueprintCosmetic)
 			&& !FunctionToCall->GetOuterUClass()->IsChildOf(UInterface::StaticClass())
-			&& FunctionToCall->GetOwnerClass()->GetName() == TEXT("KismetMathLibrary");
+			// If the function has wildcard parameters (e.g. UKismetArrayLibrary) we assume that it
+			// is doing things with reflection data at runtime and will therefore benefit from
+			// the safety of a 'context' based function call (e.g. if context is invalid or
+			// 'bArrayContextFailed' we can skip the function call):
+			&& !UEdGraphSchema_K2::HasWildcardParams(FunctionToCall);
 
 		const bool bLocalScriptFunction = 
 			!FunctionToCall->HasAnyFunctionFlags(FUNC_Native|FUNC_NetFuncFlags);
