Fixed check() on loading an anim BP when the layer interface changed without the implemetations already being loaded

#rnx

diff --git a/Engine/Source/Editor/KismetCompiler/Private/AnimBlueprintCompiler.cpp b/Engine/Source/Editor/KismetCompiler/Private/AnimBlueprintCompiler.cpp
--- a/Engine/Source/Editor/KismetCompiler/Private/AnimBlueprintCompiler.cpp
+++ b/Engine/Source/Editor/KismetCompiler/Private/AnimBlueprintCompiler.cpp
@@ -2137,14 +2137,20 @@ void FAnimBlueprintCompilerContext::ProcessSubInput(UAnimGraphNode_SubInput* InS
 						VariableGetNode->SetFromProperty(NewSubInputProperty, true);
 						VariableGetNode->AllocateDefaultPins();
 
-						// link up to new node
-						UEdGraphPin* VariablePin = VariableGetNode->FindPinChecked(NewSubInputProperty->GetFName());
-						TArray<UEdGraphPin*> Links = Pin->LinkedTo;
-						Pin->BreakAllPinLinks();
-
-						for(UEdGraphPin* LinkPin : Links)
+						// link up to new node - note that this is not a FindPinChecked because if an interface changes without the
+						// implementing class being loaded, then its graphs will not be conformed until AFTER the skeleton class
+						// has been compiled, so the variable cannot be created. This also doesnt matter, as there wont be anything connected
+						// to the pin yet anyways.
+						UEdGraphPin* VariablePin = VariableGetNode->FindPin(NewSubInputProperty->GetFName());
+						if(VariablePin)
 						{
-							VariablePin->MakeLinkTo(LinkPin);
+							TArray<UEdGraphPin*> Links = Pin->LinkedTo;
+							Pin->BreakAllPinLinks();
+
+							for(UEdGraphPin* LinkPin : Links)
+							{
+								VariablePin->MakeLinkTo(LinkPin);
+							}
 						}
 					}
 				}
