Allow faster 'math calls' even for native functions that aren't marked blueprint pure, don't allow fast calls for net functions (which *may* want context object in some configs)

diff --git a/Engine/Source/Editor/KismetCompiler/Private/KismetCompilerVMBackend.cpp b/Engine/Source/Editor/KismetCompiler/Private/KismetCompilerVMBackend.cpp
--- a/Engine/Source/Editor/KismetCompiler/Private/KismetCompilerVMBackend.cpp
+++ b/Engine/Source/Editor/KismetCompiler/Private/KismetCompilerVMBackend.cpp
@@ -1189,8 +1189,8 @@ public:
 
 		const bool bFinalFunction = FunctionToCall->HasAnyFunctionFlags(FUNC_Final) || Statement.bIsParentContext;
 		const bool bMathCall = bFinalFunction
-			&& FunctionToCall->HasAllFunctionFlags(FUNC_Static | FUNC_BlueprintPure | FUNC_Final | FUNC_Native)
-			&& !FunctionToCall->HasAnyFunctionFlags(FUNC_BlueprintAuthorityOnly | FUNC_BlueprintCosmetic)
+			&& FunctionToCall->HasAllFunctionFlags(FUNC_Static | FUNC_Final | FUNC_Native)
+			&& !FunctionToCall->HasAnyFunctionFlags(FUNC_BlueprintAuthorityOnly | FUNC_BlueprintCosmetic | FUNC_NetFuncFlags)
 			&& !FunctionToCall->GetOuterUClass()->IsChildOf(UInterface::StaticClass())
 			// If the function has wildcard parameters (e.g. UKismetArrayLibrary) we assume that it
 			// is doing things with reflection data at runtime and will therefore benefit from
