Allow setting Soft Actor references in user structs/data tables, only Hard refs are unsafe
From https://udn.unrealengine.com/questions/497734/unable-to-set-soft-references-in-data-tables.html

diff --git a/Engine/Source/Editor/KismetCompiler/Private/UserDefinedStructureCompilerUtils.cpp b/Engine/Source/Editor/KismetCompiler/Private/UserDefinedStructureCompilerUtils.cpp
--- a/Engine/Source/Editor/KismetCompiler/Private/UserDefinedStructureCompilerUtils.cpp
+++ b/Engine/Source/Editor/KismetCompiler/Private/UserDefinedStructureCompilerUtils.cpp
@@ -290,9 +290,9 @@ struct FUserDefinedStructureCompilerInner
 			{
 				const UClass* ClassObject = Cast<UClass>(VarType.PinSubCategoryObject.Get());
 
-				if (ClassObject && ClassObject->IsChildOf(AActor::StaticClass()))
+				if (ClassObject && ClassObject->IsChildOf(AActor::StaticClass()) && (VarType.PinCategory == UEdGraphSchema_K2::PC_Object || VarType.PinCategory == UEdGraphSchema_K2::PC_Interface))
 				{
-					// prevent Actor variables from having default values (because Blueprint templates are library elements that can 
+					// prevent hard reference Actor variables from having default values (because Blueprint templates are library elements that can 
 					// bridge multiple levels and different levels might not have the actor that the default is referencing).
 					VarProperty->PropertyFlags |= CPF_DisableEditOnTemplate;
 				}
