Automatic Rule transitions don't work when a Layer node is used
#fix Cache layer node indices for Bake State Machines, and Asset Player node indices for Animation Layer Graphs. Combining those at runtime to retrieve the currently active Sub Anim Instance object and any contained asset players.

diff --git a/Engine/Source/Editor/KismetCompiler/Private/AnimBlueprintCompiler.cpp b/Engine/Source/Editor/KismetCompiler/Private/AnimBlueprintCompiler.cpp
--- a/Engine/Source/Editor/KismetCompiler/Private/AnimBlueprintCompiler.cpp
+++ b/Engine/Source/Editor/KismetCompiler/Private/AnimBlueprintCompiler.cpp
@@ -53,6 +53,7 @@
 #include "AnimGraphNode_SubInput.h"
 #include "K2Node_FunctionEntry.h"
 #include "K2Node_FunctionResult.h"
+#include "AnimGraphNode_Layer.h"
 
 #define LOCTEXT_NAMESPACE "AnimBlueprintCompiler"
 
@@ -567,6 +568,51 @@ void FAnimBlueprintCompilerContext::ProcessAnimationNode(UAnimGraphNode_Base* Vi
 		// Process root nodes
 		ProcessRoot(RootNode);
 	}
+	else if (UAnimGraphNode_AssetPlayerBase* AssetPlayerBase = Cast<UAnimGraphNode_AssetPlayerBase>(VisualAnimNode))
+	{			
+		// Process Asset Player nodes to, if necessary cache off their node index for retrieval at runtime (used for evaluating Automatic Rule Transitions when using Layer nodes)
+		auto ProcessGraph = [this, GUID = AssetPlayerBase->NodeGuid](UEdGraph* Graph)
+		{
+			// Make sure we do not process the default AnimGraph
+			static const FName DefaultAnimGraphName("AnimGraph");
+			if (Graph->GetFName() != DefaultAnimGraphName)
+			{
+				FString GraphName = Graph->GetName();
+				// Also make sure we do not process any empty stub graphs
+				if (!GraphName.Contains(ANIM_FUNC_DECORATOR))
+				{
+					if (Graph->Nodes.ContainsByPredicate([GUID](UEdGraphNode* Node) { return Node->NodeGuid == GUID; }))
+					{
+						if (int32* IndexPtr = NewAnimBlueprintClass->AnimBlueprintDebugData.NodeGuidToIndexMap.Find(GUID))
+						{
+							FGraphAssetPlayerInformation& Info = NewAnimBlueprintClass->GraphAssetPlayerInformation.FindOrAdd(FName(*GraphName));
+							Info.PlayerNodeIndices.AddUnique(*IndexPtr);
+						}
+					}
+				}
+			}
+		};
+
+		// Check for any definition of a layer graph
+		for (UEdGraph* Graph : Blueprint->FunctionGraphs)
+		{
+			ProcessGraph(Graph);
+		}
+
+		// Check for any implemented AnimLayer interface graphs
+		for (FBPInterfaceDescription& InterfaceDesc : Blueprint->ImplementedInterfaces)
+		{
+			// Only process Anim Layer interfaces
+			if (InterfaceDesc.Interface->IsChildOf<UAnimLayerInterface>())
+			{
+				for (UEdGraph* Graph : InterfaceDesc.Graphs)
+				{
+					ProcessGraph(Graph);
+				}
+			}
+		}
+				
+	}
 
 	// should we do this earlier? @Tom should we split this to create anim instance vars vs linking to sub anim instance node?
 	if (UAnimGraphNode_CustomProperty* CustomPropNode = Cast<UAnimGraphNode_CustomProperty>(VisualAnimNode))
@@ -1615,25 +1661,34 @@ void FAnimBlueprintCompilerContext::ProcessStateMachine(UAnimGraphNode_StateMach
 
 		FBakedAnimationState& BakedState = Oven.GetMachine().States[StateIndex];
 
-		// Add indices to all player nodes
+		// Add indices to all player and layer nodes
 		TArray<UEdGraph*> GraphsToCheck;
-		TArray<UAnimGraphNode_AssetPlayerBase*> AssetPlayerNodes;
 		GraphsToCheck.Add(StateNode->GetBoundGraph());
 		StateNode->GetBoundGraph()->GetAllChildrenGraphs(GraphsToCheck);
 
-		for(UEdGraph* ChildGraph : GraphsToCheck)
+		TArray<UAnimGraphNode_Layer*> LayerNodes;
+		TArray<UAnimGraphNode_AssetPlayerBase*> AssetPlayerNodes;
+		for (UEdGraph* ChildGraph : GraphsToCheck)
 		{
 			ChildGraph->GetNodesOfClass(AssetPlayerNodes);
+			ChildGraph->GetNodesOfClass(LayerNodes);
 		}
 
-		for(UAnimGraphNode_AssetPlayerBase* Node : AssetPlayerNodes)
+		for (UAnimGraphNode_AssetPlayerBase* Node : AssetPlayerNodes)
 		{
-			if(int32* IndexPtr = NewAnimBlueprintClass->AnimBlueprintDebugData.NodeGuidToIndexMap.Find(Node->NodeGuid))
+			if (int32* IndexPtr = NewAnimBlueprintClass->AnimBlueprintDebugData.NodeGuidToIndexMap.Find(Node->NodeGuid))
 			{
 				BakedState.PlayerNodeIndices.Add(*IndexPtr);
 			}
 		}
 
+		for (UAnimGraphNode_Layer* Node : LayerNodes)
+		{
+			if (int32* IndexPtr = NewAnimBlueprintClass->AnimBlueprintDebugData.NodeGuidToIndexMap.Find(Node->NodeGuid))
+			{
+				BakedState.LayerNodeIndices.Add(*IndexPtr);
+			}
+		}
 		// Handle all the transitions out of this node
 		TArray<class UAnimStateTransitionNode*> TransitionList;
 		StateNode->GetTransitionList(/*out*/ TransitionList, /*bWantSortedList=*/ true);
@@ -2225,6 +2280,7 @@ void FAnimBlueprintCompilerContext::CleanAndSanitizeClass(UBlueprintGeneratedCla
 	NewAnimBlueprintClass->SubInstanceNodeProperties.Empty();
 	NewAnimBlueprintClass->LayerNodeProperties.Empty();
 	NewAnimBlueprintClass->EvaluateGraphExposedInputs.Empty();
+	NewAnimBlueprintClass->GraphAssetPlayerInformation.Empty();
 
 	// Copy over runtime data from the blueprint to the class
 	NewAnimBlueprintClass->TargetSkeleton = AnimBlueprint->TargetSkeleton;
diff --git a/Engine/Source/Runtime/Engine/Classes/Animation/AnimBlueprintGeneratedClass.h b/Engine/Source/Runtime/Engine/Classes/Animation/AnimBlueprintGeneratedClass.h
--- a/Engine/Source/Runtime/Engine/Classes/Animation/AnimBlueprintGeneratedClass.h
+++ b/Engine/Source/Runtime/Engine/Classes/Animation/AnimBlueprintGeneratedClass.h
@@ -210,6 +210,10 @@ class ENGINE_API UAnimBlueprintGeneratedClass : public UBlueprintGeneratedClass,
 	UPROPERTY()
 	TArray<FExposedValueHandler> EvaluateGraphExposedInputs;
 
+	// Indices for any Asset Player found within a specific (named) Anim Layer Graph, or implemented Anim Interface Graph
+	UPROPERTY()
+	TMap<FName, FGraphAssetPlayerInformation> GraphAssetPlayerInformation;
+
 public:
 
 	virtual const TArray<FBakedAnimationStateMachine>& GetBakedStateMachines() const override { return BakedStateMachines; }
@@ -223,6 +227,7 @@ public:
 	virtual int32 GetSyncGroupIndex(FName SyncGroupName) const override { return SyncGroupNames.IndexOfByKey(SyncGroupName); }
 	virtual const TArray<FExposedValueHandler>& GetExposedValueHandlers() const { return EvaluateGraphExposedInputs; }
 	virtual const TArray<FAnimBlueprintFunction>& GetAnimBlueprintFunctions() const override { return AnimBlueprintFunctions; }
+	virtual const TMap<FName, FGraphAssetPlayerInformation>& GetGraphAssetPlayerInformation() const override { return GraphAssetPlayerInformation; }
 	
 public:
 #if WITH_EDITORONLY_DATA
diff --git a/Engine/Source/Runtime/Engine/Classes/Animation/AnimClassData.h b/Engine/Source/Runtime/Engine/Classes/Animation/AnimClassData.h
--- a/Engine/Source/Runtime/Engine/Classes/Animation/AnimClassData.h
+++ b/Engine/Source/Runtime/Engine/Classes/Animation/AnimClassData.h
@@ -50,6 +50,10 @@ public:
 	UPROPERTY()
 	TArray<UStructProperty*> LayerNodeProperties;
 
+	// Indices for any Asset Player found within a specific (named) Anim Layer Graph, or implemented Anim Interface Graph
+	UPROPERTY()
+	TMap<FName, FGraphAssetPlayerInformation> GraphNameAssetPlayers;
+
 	// Array of sync group names in the order that they are requested during compile
 	UPROPERTY()
 	TArray<FName> SyncGroupNames;
@@ -71,6 +75,7 @@ public:
 	virtual const TArray<FName>& GetSyncGroupNames() const override { return SyncGroupNames; }
 	virtual int32 GetSyncGroupIndex(FName SyncGroupName) const override { return SyncGroupNames.IndexOfByKey(SyncGroupName); }
 	virtual const TArray<FExposedValueHandler>& GetExposedValueHandlers() const { return EvaluateGraphExposedInputs; }
+	virtual const TMap<FName, FGraphAssetPlayerInformation>& GetGraphAssetPlayerInformation() const { return GraphNameAssetPlayers; }
 
 #if WITH_EDITOR
 	void CopyFrom(IAnimClassInterface* AnimClass)
@@ -86,6 +91,7 @@ public:
 		LayerNodeProperties = AnimClass->GetLayerNodeProperties();
 		SyncGroupNames = AnimClass->GetSyncGroupNames();
 		EvaluateGraphExposedInputs = AnimClass->GetExposedValueHandlers();
+		GraphNameAssetPlayers = AnimClass->GetGraphAssetPlayerInformation();
 	}
 #endif // WITH_EDITOR
 };
diff --git a/Engine/Source/Runtime/Engine/Classes/Animation/AnimClassInterface.h b/Engine/Source/Runtime/Engine/Classes/Animation/AnimClassInterface.h
--- a/Engine/Source/Runtime/Engine/Classes/Animation/AnimClassInterface.h
+++ b/Engine/Source/Runtime/Engine/Classes/Animation/AnimClassInterface.h
@@ -95,6 +95,17 @@ struct ENGINE_API FCachedPoseIndices
 	}
 };
 
+/** Contains indices for any Asset Player nodes found for a specific Name Anim Graph (only and specifically harvested for Anim Graph Layers and Implemented Anim Layer Graphs) */
+USTRUCT()
+struct FGraphAssetPlayerInformation
+{
+	GENERATED_USTRUCT_BODY()
+
+	UPROPERTY()
+	TArray<int32> PlayerNodeIndices;
+};
+
+
 UINTERFACE()
 class ENGINE_API UAnimClassInterface : public UInterface
 {
@@ -114,7 +125,7 @@ public:
 	virtual const TArray<FName>& GetSyncGroupNames() const = 0;
 	virtual const TMap<FName, FCachedPoseIndices>& GetOrderedSavedPoseNodeIndicesMap() const = 0;
 	virtual const TArray<FAnimBlueprintFunction>& GetAnimBlueprintFunctions() const = 0;
-
+	virtual const TMap<FName, FGraphAssetPlayerInformation>& GetGraphAssetPlayerInformation() const = 0;
 	virtual USkeleton* GetTargetSkeleton() const = 0;
 
 	virtual int32 GetSyncGroupIndex(FName SyncGroupName) const = 0;
diff --git a/Engine/Source/Runtime/Engine/Classes/Animation/AnimInstance.h b/Engine/Source/Runtime/Engine/Classes/Animation/AnimInstance.h
--- a/Engine/Source/Runtime/Engine/Classes/Animation/AnimInstance.h
+++ b/Engine/Source/Runtime/Engine/Classes/Animation/AnimInstance.h
@@ -887,6 +887,9 @@ public:
 	 */
 	int32 GetInstanceAssetPlayerIndex(FName MachineName, FName StateName, FName InstanceName = NAME_None);
 
+	/** Returns all Animation Nodes of FAnimNode_AssetPlayerBase class within the specified (named) Animation Graph */
+	TArray<FAnimNode_AssetPlayerBase*> GetInstanceAssetPlayers(const FName& GraphName);
+
 	/** Gets the runtime instance desc of the state machine specified by name */
 	const FBakedAnimationStateMachine* GetStateMachineInstanceDesc(FName MachineName);
 
diff --git a/Engine/Source/Runtime/Engine/Classes/Animation/AnimStateMachineTypes.h b/Engine/Source/Runtime/Engine/Classes/Animation/AnimStateMachineTypes.h
--- a/Engine/Source/Runtime/Engine/Classes/Animation/AnimStateMachineTypes.h
+++ b/Engine/Source/Runtime/Engine/Classes/Animation/AnimStateMachineTypes.h
@@ -248,6 +248,10 @@ struct FBakedAnimationState
 	UPROPERTY()
 	TArray<int32> PlayerNodeIndices;
 
+	// Indices into the property array for layer nodes in the state
+	UPROPERTY()
+	TArray<int32> LayerNodeIndices;
+
 	// Whether or not this state will ALWAYS reset it's state on reentry, regardless of remaining weight
 	UPROPERTY()
 	bool bAlwaysResetOnEntry;
diff --git a/Engine/Source/Runtime/Engine/Private/Animation/AnimBlueprintGeneratedClass.cpp b/Engine/Source/Runtime/Engine/Private/Animation/AnimBlueprintGeneratedClass.cpp
--- a/Engine/Source/Runtime/Engine/Private/Animation/AnimBlueprintGeneratedClass.cpp
+++ b/Engine/Source/Runtime/Engine/Private/Animation/AnimBlueprintGeneratedClass.cpp
@@ -16,6 +16,7 @@
 #include "Animation/AnimNode_Root.h"
 #include "Animation/AnimNode_SubInput.h"
 #include "Animation/AnimNode_Layer.h"
+#include "Animation/AnimNode_AssetPlayerBase.h"
 
 /////////////////////////////////////////////////////
 // FStateMachineDebugData
@@ -290,10 +291,11 @@ void UAnimBlueprintGeneratedClass::Link(FArchive& Ar, bool bRelinkExistingProper
 
 	if(RootClass != this)
 	{
-		// State notifies and baked machines from the root class
+		// State notifies and baked machines and asset player information from the root class
 		check(RootClass);
 		AnimNotifies = RootClass->AnimNotifies;
 		BakedStateMachines = RootClass->BakedStateMachines;
+		GraphAssetPlayerInformation = RootClass->GraphAssetPlayerInformation;
 	}
 
 	if(RootClass != this)
diff --git a/Engine/Source/Runtime/Engine/Private/Animation/AnimInstance.cpp b/Engine/Source/Runtime/Engine/Private/Animation/AnimInstance.cpp
--- a/Engine/Source/Runtime/Engine/Private/Animation/AnimInstance.cpp
+++ b/Engine/Source/Runtime/Engine/Private/Animation/AnimInstance.cpp
@@ -2938,6 +2938,11 @@ int32 UAnimInstance::GetInstanceAssetPlayerIndex(FName MachineName, FName StateN
 	return GetProxyOnGameThread<FAnimInstanceProxy>().GetInstanceAssetPlayerIndex(MachineName, StateName, AssetName);
 }
 
+TArray<FAnimNode_AssetPlayerBase*> UAnimInstance::GetInstanceAssetPlayers(const FName& GraphName)
+{
+	return GetProxyOnGameThread<FAnimInstanceProxy>().GetInstanceAssetPlayers(GraphName);
+}
+
 FAnimNode_AssetPlayerBase* UAnimInstance::GetRelevantAssetPlayerFromState(int32 MachineIndex, int32 StateIndex)
 {
 	return GetProxyOnGameThread<FAnimInstanceProxy>().GetRelevantAssetPlayerFromState(MachineIndex, StateIndex);
diff --git a/Engine/Source/Runtime/Engine/Private/Animation/AnimInstanceProxy.cpp b/Engine/Source/Runtime/Engine/Private/Animation/AnimInstanceProxy.cpp
--- a/Engine/Source/Runtime/Engine/Private/Animation/AnimInstanceProxy.cpp
+++ b/Engine/Source/Runtime/Engine/Private/Animation/AnimInstanceProxy.cpp
@@ -23,6 +23,7 @@
 #define DO_ANIMSTAT_PROCESSING(StatName) DEFINE_STAT(STAT_ ## StatName ## _WorkerThread)
 #include "Animation/AnimMTStats.h"
 #include "Animation/AnimNode_Root.h"
+#include "Animation/AnimNode_Layer.h"
 #undef DO_ANIMSTAT_PROCESSING
 
 #define LOCTEXT_NAMESPACE "AnimInstance"
@@ -2319,6 +2320,29 @@ void FAnimInstanceProxy::ResetDynamics()
 	ResetDynamics(ETeleportType::ResetPhysics);
 }
 
+TArray<FAnimNode_AssetPlayerBase*> FAnimInstanceProxy::GetInstanceAssetPlayers(const FName& GraphName)
+{
+	TArray<FAnimNode_AssetPlayerBase*> Nodes;
+
+	// Retrieve all asset player nodes from the (named) Animation Layer Graph	
+	if (AnimClassInterface)
+	{
+		const TMap<FName, FGraphAssetPlayerInformation>& GrapInformationMap = AnimClassInterface->GetGraphAssetPlayerInformation();
+		if (const FGraphAssetPlayerInformation* Information = GrapInformationMap.Find(GraphName))
+		{
+			for (const int32& NodeIndex : Information->PlayerNodeIndices)
+			{
+				if (FAnimNode_AssetPlayerBase* Node = GetNodeFromIndex<FAnimNode_AssetPlayerBase>(NodeIndex))
+				{
+					Nodes.Add(Node);
+				}
+			}
+		}
+	}
+
+	return Nodes;
+}
+
 #if WITH_EDITOR
 void FAnimInstanceProxy::RegisterWatchedPose(const FCompactPose& Pose, int32 LinkID)
 {
diff --git a/Engine/Source/Runtime/Engine/Private/Animation/AnimNode_StateMachine.cpp b/Engine/Source/Runtime/Engine/Private/Animation/AnimNode_StateMachine.cpp
--- a/Engine/Source/Runtime/Engine/Private/Animation/AnimNode_StateMachine.cpp
+++ b/Engine/Source/Runtime/Engine/Private/Animation/AnimNode_StateMachine.cpp
@@ -6,6 +6,8 @@
 #include "Animation/AnimNode_TransitionPoseEvaluator.h"
 #include "Animation/AnimNode_AssetPlayerBase.h"
 #include "Animation/BlendProfile.h"
+#include "Animation/AnimNode_Layer.h"
+#include "Animation/AnimInstance.h"
 
 DECLARE_CYCLE_STAT(TEXT("StateMachine SetState"), Stat_StateMachineSetState, STATGROUP_Anim);
 
@@ -525,14 +527,39 @@ FAnimNode_AssetPlayerBase* FAnimNode_StateMachine::GetRelevantAssetPlayerFromSta
 {
 	FAnimNode_AssetPlayerBase* ResultPlayer = nullptr;
 	float MaxWeight = 0.0f;
+
+	auto EvaluatePlayerWeight = [&MaxWeight, &ResultPlayer](FAnimNode_AssetPlayerBase* Player)
+	{
+		if (!Player->bIgnoreForRelevancyTest && (Player->GetCachedBlendWeight() > MaxWeight))
+		{
+			MaxWeight = Player->GetCachedBlendWeight();
+			ResultPlayer = Player;
+		}
+	};
+
 	for (const int32& PlayerIdx : StateInfo.PlayerNodeIndices)
 	{
 		if (FAnimNode_AssetPlayerBase* Player = Context.AnimInstanceProxy->GetNodeFromIndex<FAnimNode_AssetPlayerBase>(PlayerIdx))
 		{
-			if (!Player->bIgnoreForRelevancyTest && (Player->GetCachedBlendWeight() > MaxWeight))
+			EvaluatePlayerWeight(Player);
+		}
+	}
+
+	// Get all layer node indices that are part of this state
+	for (const int32& LayerIdx : StateInfo.LayerNodeIndices)
+	{
+		// Try and retrieve the actual node object
+		if (FAnimNode_Layer* Layer = Context.AnimInstanceProxy->GetNodeFromIndex<FAnimNode_Layer>(LayerIdx))
+		{
+			// Retrieve the AnimInstance running for this layer
+			if (UAnimInstance* CurrentTarget = Layer->GetTargetInstance<UAnimInstance>())
 			{
-				MaxWeight = Player->GetCachedBlendWeight();
-				ResultPlayer = Player;
+				// Retrieve all asset player nodes from the corresponding Anim blueprint class and apply same logic to find highest weighted asset player 
+				TArray<FAnimNode_AssetPlayerBase*> PlayerNodesInLayer = CurrentTarget->GetInstanceAssetPlayers(Layer->Layer);
+				for (FAnimNode_AssetPlayerBase* Player : PlayerNodesInLayer)
+				{
+					EvaluatePlayerWeight(Player);
+				}
 			}
 		}
 	}
diff --git a/Engine/Source/Runtime/Engine/Public/Animation/AnimInstanceProxy.h b/Engine/Source/Runtime/Engine/Public/Animation/AnimInstanceProxy.h
--- a/Engine/Source/Runtime/Engine/Public/Animation/AnimInstanceProxy.h
+++ b/Engine/Source/Runtime/Engine/Public/Animation/AnimInstanceProxy.h
@@ -438,6 +438,9 @@ public:
 	/** Reset any dynamics running simulation-style updates (e.g. on teleport, time skip etc.) */
 	void ResetDynamics(ETeleportType InTeleportType);
 
+	/** Returns all Animation Nodes of FAnimNode_AssetPlayerBase class within the specified (named) Animation Graph */
+	TArray<FAnimNode_AssetPlayerBase*> GetInstanceAssetPlayers(const FName& GraphName);
+
 	UE_DEPRECATED(4.20, "Please use ResetDynamics with a ETeleportType argument")
 	void ResetDynamics();
 
