PR #6105: Add CLASS_HasInstancedReference flag if any property is instanced (Contributed by [REDACTED_PERSON])

Blueprints don't not support declaring properties as instanced, but this change allows native invariants to be maintained when a native struct that has properties that are marked as instanced are added to a blueprint

diff --git a/Engine/Source/Editor/KismetCompiler/Private/KismetCompiler.cpp b/Engine/Source/Editor/KismetCompiler/Private/KismetCompiler.cpp
--- a/Engine/Source/Editor/KismetCompiler/Private/KismetCompiler.cpp
+++ b/Engine/Source/Editor/KismetCompiler/Private/KismetCompiler.cpp
@@ -2454,10 +2454,18 @@ void FKismetCompilerContext::FinishCompilingClass(UClass* Class)
 	Class->ClassFlags |= (CLASS_Parsed | CLASS_CompiledFromBlueprint);
 	Class->ClassFlags &= ~CLASS_ReplicationDataIsSetUp;
 
-	// Look for OnRep 
+	// This function mostly mirrors PostParsingClassSetup, opportunity to refactor:
 	for( TFieldIterator<UProperty> It(Class, EFieldIteratorFlags::ExcludeSuper); It; ++It)
 	{
 		UProperty *Property = *It;
+		
+		// If any property is instanced, then the class needs to also have CLASS_HasInstancedReference flag
+		if (Property->ContainsInstancedObjectProperty())
+		{
+			Class->ClassFlags |= CLASS_HasInstancedReference;
+		}
+		
+		// Look for OnRep 
 		if (Property->HasAnyPropertyFlags(CPF_Net))
 		{
 			// Verify rep notifies are valid, if not, clear them
