Remove misc. dead 'skeleton class' related codepaths

diff --git a/Engine/Source/Editor/Kismet/Private/BlueprintCompilationManager.cpp b/Engine/Source/Editor/Kismet/Private/BlueprintCompilationManager.cpp
--- a/Engine/Source/Editor/Kismet/Private/BlueprintCompilationManager.cpp
+++ b/Engine/Source/Editor/Kismet/Private/BlueprintCompilationManager.cpp
@@ -2190,9 +2190,6 @@ UClass* FBlueprintCompilationManagerImpl::FastGenerateSkeletonClass(UBlueprint*
 	// Temporarily set the compile type to indicate that we're generating the skeleton class.
 	TGuardValue<EKismetCompileType::Type> GuardCompileType(CompilerContext.CompileOptions.CompileType, EKismetCompileType::SkeletonOnly);
 
-	// Temporarily set the flag to indicate whether or not we'll only be generating the skeleton class as part of the overall compile request.
-	TGuardValue<bool> GuardSkeletonOnly(CompilerContext.bIsSkeletonOnly, bIsSkeletonOnly);
-
 	if (BP->SkeletonGeneratedClass == nullptr)
 	{
 		// This might exist in the package because we are being reloaded in place
diff --git a/Engine/Source/Editor/KismetCompiler/Private/AnimBlueprintCompiler.cpp b/Engine/Source/Editor/KismetCompiler/Private/AnimBlueprintCompiler.cpp
--- a/Engine/Source/Editor/KismetCompiler/Private/AnimBlueprintCompiler.cpp
+++ b/Engine/Source/Editor/KismetCompiler/Private/AnimBlueprintCompiler.cpp
@@ -770,12 +770,6 @@ void FAnimBlueprintCompilerContext::ProcessUseCachedPose(UAnimGraphNode_UseCache
 {
 	bool bSuccessful = false;
 
-	// if compiling only skeleton, we don't have to worry about linking save node
-	if (CompileOptions.CompileType == EKismetCompileType::SkeletonOnly)
-	{
-		return;
-	}
-
 	// Link to the saved cached pose
 	if(UseCachedPose->SaveCachedPoseNode.IsValid())
 	{
@@ -1313,11 +1307,8 @@ void FAnimBlueprintCompilerContext::ProcessAllAnimationNodes()
 		MessageLog.Error(*LOCTEXT("ExpectedAFunctionEntry_Error", "Expected at least one animation root, but did not find any").ToString());
 	}
 
-	if(CompileOptions.CompileType != EKismetCompileType::SkeletonOnly)
-	{
-		// Build cached pose map
-		BuildCachedPoseNodeUpdateOrder();
-	}
+	// Build cached pose map
+	BuildCachedPoseNodeUpdateOrder();
 }
 
 int32 FAnimBlueprintCompilerContext::ExpandGraphAndProcessNodes(UEdGraph* SourceGraph, UAnimGraphNode_Base* SourceRootNode, UAnimStateTransitionNode* TransitionNode, TArray<UEdGraphNode*>* ClonedNodes)
diff --git a/Engine/Source/Editor/KismetCompiler/Private/KismetCompiler.cpp b/Engine/Source/Editor/KismetCompiler/Private/KismetCompiler.cpp
--- a/Engine/Source/Editor/KismetCompiler/Private/KismetCompiler.cpp
+++ b/Engine/Source/Editor/KismetCompiler/Private/KismetCompiler.cpp
@@ -132,7 +132,6 @@ FKismetCompilerContext::FKismetCompilerContext(UBlueprint* SourceSketch, FCompil
 	, ConsolidatedEventGraph(NULL)
 	, UbergraphContext(NULL)
 	, bIsFullCompile(false)
-	, bIsSkeletonOnly(false)
 	, OldCDO(nullptr)
 	, OldGenLinkerIdx(INDEX_NONE)
 	, OldLinker(nullptr)
@@ -3753,6 +3752,7 @@ void FKismetCompilerContext::ProcessOneFunctionGraph(UEdGraph* SourceGraph, bool
 	//       skeleton-only compiles (that's why we have that check second) 
 	//       because it would most likely result in errors (the function hasn't
 	//       been added to the class yet, etc.)
+	check(CompileOptions.CompileType != EKismetCompileType::SkeletonOnly);
 	if ((CompileOptions.CompileType == EKismetCompileType::SkeletonOnly) || ValidateGraphIsWellFormed(FunctionGraph))
 	{
 		const UEdGraphSchema_K2* FunctionGraphSchema = CastChecked<const UEdGraphSchema_K2>(FunctionGraph->GetSchema());
@@ -3898,8 +3898,7 @@ void FKismetCompilerContext::CompileClassLayout(EInternalCompilerFlags InternalF
 	// Make sure the parent class exists and can be used
 	check(Blueprint->ParentClass && Blueprint->ParentClass->GetPropertiesSize());
 
-	bIsSkeletonOnly = (CompileOptions.CompileType == EKismetCompileType::SkeletonOnly);
-	UClass* TargetUClass = bIsSkeletonOnly ? Blueprint->SkeletonGeneratedClass : Blueprint->GeneratedClass;
+	UClass* TargetUClass = Blueprint->GeneratedClass;
 
 	// >>> Backwards Compatibility:  Make sure this is an actual UBlueprintGeneratedClass / UAnimBlueprintGeneratedClass, as opposed to the old UClass
 	EnsureProperGeneratedClass(TargetUClass);
@@ -3907,45 +3906,8 @@ void FKismetCompilerContext::CompileClassLayout(EInternalCompilerFlags InternalF
 
 	TargetClass = Cast<UBlueprintGeneratedClass>(TargetUClass);
 
-	// >>> Backwards Compatibility: Make sure that skeleton generated classes have the proper "SKEL_" naming convention
-	const FString SkeletonPrefix(TEXT("SKEL_"));
-	if( bIsSkeletonOnly && TargetClass && !TargetClass->GetName().StartsWith(SkeletonPrefix) )
-	{
-		FString NewName = SkeletonPrefix + TargetClass->GetName();
- 
-		// Ensure we have a free name for this class
-		UClass* AnyClassWithGoodName = (UClass*)StaticFindObject(UClass::StaticClass(), Blueprint->GetOutermost(), *NewName, false);
-		if( AnyClassWithGoodName )
-		{
-			// Special Case:  If the CDO of the class has become dissociated from its actual CDO, attempt to find the proper named CDO, and get rid of it.
-			if( AnyClassWithGoodName->ClassDefaultObject == TargetClass->ClassDefaultObject )
-			{
-				AnyClassWithGoodName->ClassDefaultObject = NULL;
-				FString DefaultObjectName = FString(DEFAULT_OBJECT_PREFIX) + NewName;
-				AnyClassWithGoodName->ClassDefaultObject = (UObject*)StaticFindObject(UObject::StaticClass(), Blueprint->GetOutermost(), *DefaultObjectName, false);
-			}
- 
-			// Get rid of the old class to make room for renaming our class to the final SKEL name
-			FKismetCompilerUtilities::ConsignToOblivion(AnyClassWithGoodName, Blueprint->bIsRegeneratingOnLoad);
-
-			// Update the refs to the old SKC
-			TMap<UObject*, UObject*> ClassReplacementMap;
-			ClassReplacementMap.Add(AnyClassWithGoodName, TargetClass);
-			TArray<UEdGraph*> AllGraphs;
-			Blueprint->GetAllGraphs(AllGraphs);
-			for (int32 i = 0; i < AllGraphs.Num(); ++i)
-			{
-				FArchiveReplaceObjectRef<UObject> ReplaceInBlueprintAr(AllGraphs[i], ClassReplacementMap, /*bNullPrivateRefs=*/ false, /*bIgnoreOuterRef=*/ false, /*bIgnoreArchetypeRef=*/ false);
-			}
-		}
- 
-		ERenameFlags RenameFlags = (REN_DontCreateRedirectors|REN_NonTransactional|((Blueprint->bIsRegeneratingOnLoad)? REN_ForceNoResetLoaders : 0));
-		TargetClass->Rename(*NewName, NULL, RenameFlags);
-	}
-	// <<< End Backwards Compatibility
-
 	// >>> Backwards compatibility:  If SkeletonGeneratedClass == GeneratedClass, we need to make a new generated class the first time we need it
-	if( !bIsSkeletonOnly && (Blueprint->SkeletonGeneratedClass == Blueprint->GeneratedClass) )
+	if( Blueprint->SkeletonGeneratedClass == Blueprint->GeneratedClass )
 	{
 		Blueprint->GeneratedClass = NULL;
 		TargetClass = NULL;
@@ -3956,20 +3918,13 @@ void FKismetCompilerContext::CompileClassLayout(EInternalCompilerFlags InternalF
 	{
 		FName NewSkelClassName, NewGenClassName;
 		Blueprint->GetBlueprintClassNames(NewGenClassName, NewSkelClassName);
-		SpawnNewClass( bIsSkeletonOnly ? NewSkelClassName.ToString() : NewGenClassName.ToString() );
+		SpawnNewClass( NewGenClassName.ToString() );
 		check(NewClass);
 
 		TargetClass = NewClass;
 
 		// Fix up the reference in the blueprint to the new class
-		if( bIsSkeletonOnly )
-		{
-			Blueprint->SkeletonGeneratedClass = TargetClass;
-		}
-		else
-		{
-			Blueprint->GeneratedClass = TargetClass;
-		}
+		Blueprint->GeneratedClass = TargetClass;
 	}
 
 	// Early validation
@@ -4405,20 +4360,6 @@ void FKismetCompilerContext::CompileFunctions(EInternalCompilerFlags InternalFla
 		}
 	}
 
-	// If this was a skeleton compile, make sure everything is RF_Transient
-	if (bIsSkeletonOnly)
-	{
-		ForEachObjectWithOuter(NewClass, [](UObject* Child)
-		{
-			Child->SetFlags(RF_Transient);
-		});
-
-		NewClass->SetFlags(RF_Transient);
-
-		check(NewClass->ClassDefaultObject != nullptr);
-		NewClass->ClassDefaultObject->SetFlags(RF_Transient);
-	}
-
 	// For full compiles, find other blueprints that may need refreshing, and mark them dirty, in case they try to run
 	if( bIsFullCompile && !Blueprint->bIsRegeneratingOnLoad && !bSkipRefreshExternalBlueprintDependencyNodes )
 	{
diff --git a/Engine/Source/Editor/KismetCompiler/Public/KismetCompiler.h b/Engine/Source/Editor/KismetCompiler/Public/KismetCompiler.h
--- a/Engine/Source/Editor/KismetCompiler/Public/KismetCompiler.h
+++ b/Engine/Source/Editor/KismetCompiler/Public/KismetCompiler.h
@@ -119,7 +119,6 @@ public:
 	FNetNameMapping ClassScopeNetNameMap;
 
 	// Data that persists across CompileClassLayout/CompileFunctions calls:
-	bool bIsSkeletonOnly;
 	UObject* OldCDO;
 	int32 OldGenLinkerIdx;
 	FLinkerLoad* OldLinker;
diff --git a/Engine/Source/Editor/UMGEditor/Private/WidgetBlueprintCompiler.cpp b/Engine/Source/Editor/UMGEditor/Private/WidgetBlueprintCompiler.cpp
--- a/Engine/Source/Editor/UMGEditor/Private/WidgetBlueprintCompiler.cpp
+++ b/Engine/Source/Editor/UMGEditor/Private/WidgetBlueprintCompiler.cpp
@@ -609,9 +609,10 @@ void FWidgetBlueprintCompilerContext::FinishCompilingClass(UClass* Class)
 	UWidgetBlueprint* WidgetBP = WidgetBlueprint();
 	UWidgetBlueprintGeneratedClass* BPGClass = CastChecked<UWidgetBlueprintGeneratedClass>(Class);
 	UClass* ParentClass = WidgetBP->ParentClass;
+	const bool bIsSkeletonOnly = CompileOptions.CompileType == EKismetCompileType::SkeletonOnly;
 
 	// Don't do a bunch of extra work on the skeleton generated class.
-	if ( CompileOptions.CompileType != EKismetCompileType::SkeletonOnly )
+	if ( !bIsSkeletonOnly )
 	{
 		if( !WidgetBP->bHasBeenRegenerated )
 		{
