Copying //UE4/Dev-Anim to Dev-Main (//UE4/Dev-Main)

diff --git a/Engine/Source/Editor/EditorStyle/Private/SlateEditorStyle.cpp b/Engine/Source/Editor/EditorStyle/Private/SlateEditorStyle.cpp
--- a/Engine/Source/Editor/EditorStyle/Private/SlateEditorStyle.cpp
+++ b/Engine/Source/Editor/EditorStyle/Private/SlateEditorStyle.cpp
@@ -5720,6 +5720,10 @@ void FSlateEditorStyle::FStyle::SetupPersonaStyle()
 		// blend space
 		Set("BlendSpaceEditor.ToggleTriangulation", new IMAGE_BRUSH(TEXT("Persona/BlendSpace/triangulation_16"), Icon16x16));
 		Set("BlendSpaceEditor.ToggleLabels", new IMAGE_BRUSH(TEXT("Persona/BlendSpace/label_16"), Icon16x16));
+		Set("BlendSpaceEditor.ArrowDown", new IMAGE_BRUSH(TEXT("Persona/BlendSpace/arrow_down_12x"), Icon12x12));
+		Set("BlendSpaceEditor.ArrowUp", new IMAGE_BRUSH(TEXT("Persona/BlendSpace/arrow_up_12x"), Icon12x12));
+		Set("BlendSpaceEditor.ArrowRight", new IMAGE_BRUSH(TEXT("Persona/BlendSpace/arrow_right_12x"), Icon12x12));
+		Set("BlendSpaceEditor.ArrowLeft", new IMAGE_BRUSH(TEXT("Persona/BlendSpace/arrow_left_12x"), Icon12x12));
 
 		const FButtonStyle BlueprintContextTargetsButtonStyle = FButtonStyle()
 			.SetNormal(IMAGE_BRUSH("Common/TreeArrow_Collapsed_Hovered", Icon10x10, FLinearColor(0.2f, 0.2f, 0.2f, 1.f)))
diff --git a/Engine/Source/Editor/KismetCompiler/Private/AnimBlueprintCompiler.cpp b/Engine/Source/Editor/KismetCompiler/Private/AnimBlueprintCompiler.cpp
--- a/Engine/Source/Editor/KismetCompiler/Private/AnimBlueprintCompiler.cpp
+++ b/Engine/Source/Editor/KismetCompiler/Private/AnimBlueprintCompiler.cpp
@@ -267,7 +267,7 @@ UK2Node_CallFunction* FAnimBlueprintCompilerContext::SpawnCallAnimInstanceFuncti
 	return FunctionCall;
 }
 
-void FAnimBlueprintCompilerContext::CreateEvaluationHandlerStruct(UAnimGraphNode_Base* VisualAnimNode, FEvaluationHandlerRecord& Record)
+void FAnimBlueprintCompilerContext::CreateEvaluationHandler(UAnimGraphNode_Base* VisualAnimNode, FEvaluationHandlerRecord& Record)
 {
 	// Shouldn't create a handler if there is nothing to work with
 	check(Record.ServicedProperties.Num() > 0);
@@ -286,7 +286,7 @@ void FAnimBlueprintCompilerContext::CreateEvaluationHandlerStruct(UAnimGraphNode
 	// check function name isnt already used (data exists that can contain duplicate GUIDs) and apply a numeric extension until it is unique
 	int32 ExtensionIndex = 0;
 	FName* ExistingName = HandlerFunctionNames.Find(Record.HandlerFunctionName);
-	while(ExistingName != nullptr)
+	while (ExistingName != nullptr)
 	{
 		FunctionName = FString::Printf(TEXT("%s_%s_%s_%s_%d"), *AnimGraphDefaultSchema->DefaultEvaluationHandlerName.ToString(), *VisualAnimNode->GetOuter()->GetName(), *VisualAnimNode->GetClass()->GetName(), *VisualAnimNode->NodeGuid.ToString(), ExtensionIndex);
 		Record.HandlerFunctionName = FName(*FunctionName);
@@ -295,7 +295,7 @@ void FAnimBlueprintCompilerContext::CreateEvaluationHandlerStruct(UAnimGraphNode
 	}
 
 	HandlerFunctionNames.Add(Record.HandlerFunctionName);
-	
+
 	// Add a custom event in the graph
 	UK2Node_CustomEvent* EntryNode = SpawnIntermediateEventNode<UK2Node_CustomEvent>(VisualAnimNode, nullptr, ConsolidatedEventGraph);
 	EntryNode->bInternalEvent = true;
@@ -304,187 +304,159 @@ void FAnimBlueprintCompilerContext::CreateEvaluationHandlerStruct(UAnimGraphNode
 
 	// The ExecChain is the current exec output pin in the linear chain
 	UEdGraphPin* ExecChain = Schema->FindExecutionPin(*EntryNode, EGPD_Output);
+	if (Record.bServicesInstanceProperties)
+	{
+		// Need to create a variable set call for each serviced property in the handler
+		for (TPair<FName, FAnimNodeSinglePropertyHandler>& PropHandlerPair : Record.ServicedProperties)
+		{
+			FAnimNodeSinglePropertyHandler& PropHandler = PropHandlerPair.Value;
+			FName PropertyName = PropHandlerPair.Key;
 
-	// Create a struct member write node to store the parameters into the animation node
-	UK2Node_StructMemberSet* AssignmentNode = SpawnIntermediateNode<UK2Node_StructMemberSet>(VisualAnimNode, ConsolidatedEventGraph);
-	AssignmentNode->VariableReference.SetSelfMember(Record.NodeVariableProperty->GetFName());
-	AssignmentNode->StructType = Record.NodeVariableProperty->Struct;
-	AssignmentNode->AllocateDefaultPins();
+			// Should be true, we only want to deal with instance targets in here
+			if (PropHandler.bInstanceIsTarget)
+			{
+				for (FPropertyCopyRecord& CopyRecord : PropHandler.CopyRecords)
+				{
+					// New set node for the property
+					UK2Node_VariableSet* VarAssignNode = SpawnIntermediateNode<UK2Node_VariableSet>(VisualAnimNode, ConsolidatedEventGraph);
+					VarAssignNode->VariableReference.SetSelfMember(CopyRecord.DestProperty->GetFName());
+					VarAssignNode->AllocateDefaultPins();
 
-	// Wire up the variable node execution wires
-	UEdGraphPin* ExecVariablesIn = Schema->FindExecutionPin(*AssignmentNode, EGPD_Input);
-	ExecChain->MakeLinkTo(ExecVariablesIn);
-	ExecChain = Schema->FindExecutionPin(*AssignmentNode, EGPD_Output);
+					// Wire up the exec line, and update the end of the chain
+					UEdGraphPin* ExecVariablesIn = Schema->FindExecutionPin(*VarAssignNode, EGPD_Input);
+					ExecChain->MakeLinkTo(ExecVariablesIn);
+					ExecChain = Schema->FindExecutionPin(*VarAssignNode, EGPD_Output);
 
-	// Run thru each property
-	TSet<FName> PropertiesBeingSet;
+					// Find the property pin on the set node and configure
+					for (UEdGraphPin* TargetPin : VarAssignNode->Pins)
+					{
+						if (TargetPin->PinType.IsContainer())
+						{
+							// Currently unsupported
+							continue;
+						}
 
-	for (auto TargetPinIt = AssignmentNode->Pins.CreateIterator(); TargetPinIt; ++TargetPinIt)
-	{
-		UEdGraphPin* TargetPin = *TargetPinIt;
-		FName PropertyName(TargetPin->PinName);
+						FName PinPropertyName(TargetPin->PinName);
 
-		// Does it get serviced by this handler?
-		if (FAnimNodeSinglePropertyHandler* SourceInfo = Record.ServicedProperties.Find(PropertyName))
-		{
-			if (TargetPin->PinType.IsArray())
-			{
-				// Grab the array that we need to set members for
-				UK2Node_StructMemberGet* FetchArrayNode = SpawnIntermediateNode<UK2Node_StructMemberGet>(VisualAnimNode, ConsolidatedEventGraph);
-				FetchArrayNode->VariableReference.SetSelfMember(Record.NodeVariableProperty->GetFName());
-				FetchArrayNode->StructType = Record.NodeVariableProperty->Struct;
-				FetchArrayNode->AllocatePinsForSingleMemberGet(PropertyName);
+						if (PinPropertyName == PropertyName)
+						{
+							// This is us, wire up the variable
+							UEdGraphPin* DestPin = CopyRecord.DestPin;
 
-				UEdGraphPin* ArrayVariableNode = FetchArrayNode->FindPin(PropertyName);
+							// Copy the data (link up to the source nodes)
+							TargetPin->CopyPersistentDataFromOldPin(*DestPin);
+							MessageLog.NotifyIntermediatePinCreation(TargetPin, DestPin);
 
-				if (SourceInfo->CopyRecords.Num() > 0)
-				{
-					// Set each element in the array
-					for (FPropertyCopyRecord& CopyRecord : SourceInfo->CopyRecords)
-					{
-						int32 ArrayIndex = CopyRecord.DestArrayIndex;
-						UEdGraphPin* DestPin = CopyRecord.DestPin;
-
-						// Create an array element set node
-						UK2Node_CallArrayFunction* ArrayNode = SpawnIntermediateNode<UK2Node_CallArrayFunction>(VisualAnimNode, ConsolidatedEventGraph);
-						ArrayNode->FunctionReference.SetExternalMember(GET_FUNCTION_NAME_CHECKED(UKismetArrayLibrary, Array_Set), UKismetArrayLibrary::StaticClass());
-						ArrayNode->AllocateDefaultPins();
-
-						// Connect the execution chain
-						ExecChain->MakeLinkTo(ArrayNode->GetExecPin());
-						ExecChain = ArrayNode->GetThenPin();
-
-						// Connect the input array
-						UEdGraphPin* TargetArrayPin = ArrayNode->FindPinChecked(TEXT("TargetArray"));
-						TargetArrayPin->MakeLinkTo(ArrayVariableNode);
-						ArrayNode->PinConnectionListChanged(TargetArrayPin);
-
-						// Set the array index
-						UEdGraphPin* TargetIndexPin = ArrayNode->FindPinChecked(TEXT("Index"));
-						TargetIndexPin->DefaultValue = FString::FromInt(ArrayIndex);
-
-						// Wire up the data input
-						UEdGraphPin* TargetItemPin = ArrayNode->FindPinChecked(TEXT("Item"));
-						TargetItemPin->CopyPersistentDataFromOldPin(*DestPin);
-						MessageLog.NotifyIntermediatePinCreation(TargetItemPin, DestPin);
-						DestPin->BreakAllPinLinks();
+							// Old pin needs to not be connected now - break all its links
+							DestPin->BreakAllPinLinks();
+
+							break;
+						}
 					}
-				}
-			}
-			else
-			{
-				check(!TargetPin->PinType.IsContainer())
-				// Single property
-				if (SourceInfo->CopyRecords.Num() > 0 && SourceInfo->CopyRecords[0].DestPin != nullptr)
-				{
-					UEdGraphPin* DestPin = SourceInfo->CopyRecords[0].DestPin;
 
-					PropertiesBeingSet.Add(DestPin->PinName);
-					TargetPin->CopyPersistentDataFromOldPin(*DestPin);
-					MessageLog.NotifyIntermediatePinCreation(TargetPin, DestPin);
-					DestPin->BreakAllPinLinks();
+					//VarAssignNode->ReconstructNode();
 				}
 			}
 		}
 	}
 
-	// Remove any unused pins from the assignment node to avoid smashing constant values
-	for (int32 PinIndex = 0; PinIndex < AssignmentNode->ShowPinForProperties.Num(); ++PinIndex)
+	if (Record.bServicesNodeProperties)
 	{
-		FOptionalPinFromProperty& TestProperty = AssignmentNode->ShowPinForProperties[PinIndex];
-		TestProperty.bShowPin = PropertiesBeingSet.Contains(TestProperty.PropertyName);
-	}
-	AssignmentNode->ReconstructNode();
-}
+		// Create a struct member write node to store the parameters into the animation node
+		UK2Node_StructMemberSet* AssignmentNode = SpawnIntermediateNode<UK2Node_StructMemberSet>(VisualAnimNode, ConsolidatedEventGraph);
+		AssignmentNode->VariableReference.SetSelfMember(Record.NodeVariableProperty->GetFName());
+		AssignmentNode->StructType = Record.NodeVariableProperty->Struct;
+		AssignmentNode->AllocateDefaultPins();
 
-void FAnimBlueprintCompilerContext::CreateEvaluationHandlerInstance(UAnimGraphNode_Base* VisualAnimNode, FEvaluationHandlerRecord& Record)
-{
-	// Shouldn't create a handler if there is nothing to work with
-	check(Record.ServicedProperties.Num() > 0);
-	check(Record.NodeVariableProperty != nullptr);
-	check(Record.bServicesInstanceProperties);
-	
-	const UAnimationGraphSchema* AnimGraphDefaultSchema = GetDefault<UAnimationGraphSchema>();
+		// Wire up the variable node execution wires
+		UEdGraphPin* ExecVariablesIn = Schema->FindExecutionPin(*AssignmentNode, EGPD_Input);
+		ExecChain->MakeLinkTo(ExecVariablesIn);
+		ExecChain = Schema->FindExecutionPin(*AssignmentNode, EGPD_Output);
 
-	// Use the node GUID for a stable name across compiles
-	FString FunctionName = FString::Printf(TEXT("%s_%s_%s_%s"), *AnimGraphDefaultSchema->DefaultEvaluationHandlerName.ToString(), *VisualAnimNode->GetOuter()->GetName(), *VisualAnimNode->GetClass()->GetName(), *VisualAnimNode->NodeGuid.ToString());
-	Record.HandlerFunctionName = FName(*FunctionName);
+		// Run thru each property
+		TSet<FName> PropertiesBeingSet;
 
-	// check function name isnt already used (data exists that can contain duplicate GUIDs) and apply a numeric extension until it is unique
-	int32 ExtensionIndex = 0;
-	FName* ExistingName = HandlerFunctionNames.Find(Record.HandlerFunctionName);
-	while(ExistingName != nullptr)
-	{
-		FunctionName = FString::Printf(TEXT("%s_%s_%s_%s_%d"), *AnimGraphDefaultSchema->DefaultEvaluationHandlerName.ToString(), *VisualAnimNode->GetOuter()->GetName(), *VisualAnimNode->GetClass()->GetName(), *VisualAnimNode->NodeGuid.ToString(), ExtensionIndex);
-		Record.HandlerFunctionName = FName(*FunctionName);
-		ExistingName = HandlerFunctionNames.Find(Record.HandlerFunctionName);
-		ExtensionIndex++;
-	}
-
-	HandlerFunctionNames.Add(Record.HandlerFunctionName);
-
-	// Add a custom event in the graph
-	UK2Node_CustomEvent* EntryNode = SpawnIntermediateNode<UK2Node_CustomEvent>(VisualAnimNode, ConsolidatedEventGraph);
-	EntryNode->bInternalEvent = true;
-	EntryNode->CustomFunctionName = Record.HandlerFunctionName;
-	EntryNode->AllocateDefaultPins();
-
-	// The ExecChain is the current exec output pin in the linear chain
-	UEdGraphPin* ExecChain = Schema->FindExecutionPin(*EntryNode, EGPD_Output);
-
-	// Need to create a variable set call for each serviced property in the handler
-	for(TPair<FName, FAnimNodeSinglePropertyHandler>& PropHandlerPair : Record.ServicedProperties)
-	{
-		FAnimNodeSinglePropertyHandler& PropHandler = PropHandlerPair.Value;
-		FName PropertyName = PropHandlerPair.Key;
-
-		// Should be true, we only want to deal with instance targets in here
-		check(PropHandler.bInstanceIsTarget);
-
-		for(FPropertyCopyRecord& CopyRecord : PropHandler.CopyRecords)
+		for (auto TargetPinIt = AssignmentNode->Pins.CreateIterator(); TargetPinIt; ++TargetPinIt)
 		{
-			// New set node for the property
-			UK2Node_VariableSet* VarAssignNode = SpawnIntermediateNode<UK2Node_VariableSet>(VisualAnimNode, ConsolidatedEventGraph);
-			VarAssignNode->VariableReference.SetSelfMember(CopyRecord.DestProperty->GetFName());
-			VarAssignNode->AllocateDefaultPins();
-
-			// Wire up the exec line, and update the end of the chain
-			UEdGraphPin* ExecVariablesIn = Schema->FindExecutionPin(*VarAssignNode, EGPD_Input);
-			ExecChain->MakeLinkTo(ExecVariablesIn);
-			ExecChain = Schema->FindExecutionPin(*VarAssignNode, EGPD_Output);
+			UEdGraphPin* TargetPin = *TargetPinIt;
+			FName PropertyName(TargetPin->PinName);
 
-			// Find the property pin on the set node and configure
-			for(UEdGraphPin* TargetPin : VarAssignNode->Pins)
+			// Does it get serviced by this handler?
+			if (FAnimNodeSinglePropertyHandler* SourceInfo = Record.ServicedProperties.Find(PropertyName))
 			{
-				if(TargetPin->PinType.IsContainer())
+				if (TargetPin->PinType.IsArray())
 				{
-					// Currently unsupported
-					continue;
-				}
+					// Grab the array that we need to set members for
+					UK2Node_StructMemberGet* FetchArrayNode = SpawnIntermediateNode<UK2Node_StructMemberGet>(VisualAnimNode, ConsolidatedEventGraph);
+					FetchArrayNode->VariableReference.SetSelfMember(Record.NodeVariableProperty->GetFName());
+					FetchArrayNode->StructType = Record.NodeVariableProperty->Struct;
+					FetchArrayNode->AllocatePinsForSingleMemberGet(PropertyName);
 
-				FName PinPropertyName(TargetPin->PinName);
+					UEdGraphPin* ArrayVariableNode = FetchArrayNode->FindPin(PropertyName);
 
-				if(PinPropertyName == PropertyName)
+					if (SourceInfo->CopyRecords.Num() > 0)
+					{
+						// Set each element in the array
+						for (FPropertyCopyRecord& CopyRecord : SourceInfo->CopyRecords)
+						{
+							int32 ArrayIndex = CopyRecord.DestArrayIndex;
+							UEdGraphPin* DestPin = CopyRecord.DestPin;
+
+							// Create an array element set node
+							UK2Node_CallArrayFunction* ArrayNode = SpawnIntermediateNode<UK2Node_CallArrayFunction>(VisualAnimNode, ConsolidatedEventGraph);
+							ArrayNode->FunctionReference.SetExternalMember(GET_FUNCTION_NAME_CHECKED(UKismetArrayLibrary, Array_Set), UKismetArrayLibrary::StaticClass());
+							ArrayNode->AllocateDefaultPins();
+
+							// Connect the execution chain
+							ExecChain->MakeLinkTo(ArrayNode->GetExecPin());
+							ExecChain = ArrayNode->GetThenPin();
+
+							// Connect the input array
+							UEdGraphPin* TargetArrayPin = ArrayNode->FindPinChecked(TEXT("TargetArray"));
+							TargetArrayPin->MakeLinkTo(ArrayVariableNode);
+							ArrayNode->PinConnectionListChanged(TargetArrayPin);
+
+							// Set the array index
+							UEdGraphPin* TargetIndexPin = ArrayNode->FindPinChecked(TEXT("Index"));
+							TargetIndexPin->DefaultValue = FString::FromInt(ArrayIndex);
+
+							// Wire up the data input
+							UEdGraphPin* TargetItemPin = ArrayNode->FindPinChecked(TEXT("Item"));
+							TargetItemPin->CopyPersistentDataFromOldPin(*DestPin);
+							MessageLog.NotifyIntermediatePinCreation(TargetItemPin, DestPin);
+							DestPin->BreakAllPinLinks();
+						}
+					}
+				}
+				else
 				{
-					// This is us, wire up the variable
-					UEdGraphPin* DestPin = CopyRecord.DestPin;
-
-					// Copy the data (link up to the source nodes)
-					TargetPin->CopyPersistentDataFromOldPin(*DestPin);
-					MessageLog.NotifyIntermediatePinCreation(TargetPin, DestPin);
-
-					// Old pin needs to not be connected now - break all its links
-					DestPin->BreakAllPinLinks();
+					check(!TargetPin->PinType.IsContainer());
+					// Single property
+					if (SourceInfo->CopyRecords.Num() > 0 && SourceInfo->CopyRecords[0].DestPin != nullptr)
+					{
+						UEdGraphPin* DestPin = SourceInfo->CopyRecords[0].DestPin;
 
-					break;
+						PropertiesBeingSet.Add(DestPin->PinName);
+						TargetPin->CopyPersistentDataFromOldPin(*DestPin);
+						MessageLog.NotifyIntermediatePinCreation(TargetPin, DestPin);
+						DestPin->BreakAllPinLinks();
+					}
 				}
 			}
+		}
 
-			//VarAssignNode->ReconstructNode();
+		// Remove any unused pins from the assignment node to avoid smashing constant values
+		for (int32 PinIndex = 0; PinIndex < AssignmentNode->ShowPinForProperties.Num(); ++PinIndex)
+		{
+			FOptionalPinFromProperty& TestProperty = AssignmentNode->ShowPinForProperties[PinIndex];
+			TestProperty.bShowPin = PropertiesBeingSet.Contains(TestProperty.PropertyName);
 		}
+
+		AssignmentNode->ReconstructNode();
 	}
 }
 
+
 void FAnimBlueprintCompilerContext::ProcessAnimationNode(UAnimGraphNode_Base* VisualAnimNode)
 {
 	// Early out if this node has already been processed
@@ -645,20 +617,20 @@ void FAnimBlueprintCompilerContext::ProcessAnimationNode(UAnimGraphNode_Base* Vi
 			// owning anim instance if using a linked instance node.
 			UProperty* SourcePinProperty = nullptr;
 			int32 SourceArrayIndex = INDEX_NONE;
+			bool bInstancePropertyExists = false;
 
 			// We have special handling below if we're targeting a linked instance instead of our own instance properties
 			UAnimGraphNode_CustomProperty* CustomPropertyNode = Cast<UAnimGraphNode_CustomProperty>(VisualAnimNode);
 
+			VisualAnimNode->GetPinAssociatedProperty(NodeType, SourcePin, /*out*/ SourcePinProperty, /*out*/ SourceArrayIndex);
+
 			// Does this pin have an associated evaluation handler?
-			if(CustomPropertyNode)
+			if(!SourcePinProperty && CustomPropertyNode)
 			{
 				// Custom property nodes use instance properties not node properties as they aren't UObjects
 				// and we can't store non-native properties there
 				CustomPropertyNode->GetInstancePinProperty(NewAnimBlueprintClass, SourcePin, SourcePinProperty);
-			}
-			else
-			{
-				VisualAnimNode->GetPinAssociatedProperty(NodeType, SourcePin, /*out*/ SourcePinProperty, /*out*/ SourceArrayIndex);
+				bInstancePropertyExists = true;
 			}
 			
 			if (SourcePinProperty != NULL)
@@ -687,16 +659,18 @@ void FAnimBlueprintCompilerContext::ProcessAnimationNode(UAnimGraphNode_Base* Vi
 					ensure(EvalHandler.NodeVariableProperty == nullptr || EvalHandler.NodeVariableProperty == NewProperty);
 					EvalHandler.NodeVariableProperty = NewProperty;
 					EvalHandler.RegisterPin(SourcePin, SourcePinProperty, SourceArrayIndex);
+					// if it's not instance property, ensure we mark it
+					EvalHandler.bServicesNodeProperties = EvalHandler.bServicesNodeProperties | !bInstancePropertyExists;
 
-					if(CustomPropertyNode)
+					if (CustomPropertyNode)
 					{
-						EvalHandler.bServicesInstanceProperties = true;
+						EvalHandler.bServicesInstanceProperties = EvalHandler.bServicesInstanceProperties | bInstancePropertyExists;
 
 						FAnimNodeSinglePropertyHandler* SinglePropHandler = EvalHandler.ServicedProperties.Find(SourcePinProperty->GetFName());
 						check(SinglePropHandler); // Should have been added in RegisterPin
 
 						// Flag that the target property is actually on the instance class and not the node
-						SinglePropHandler->bInstanceIsTarget = true;
+						SinglePropHandler->bInstanceIsTarget = bInstancePropertyExists;
 					}
 
 					bConsumed = true;
@@ -740,15 +714,7 @@ void FAnimBlueprintCompilerContext::ProcessAnimationNode(UAnimGraphNode_Base* Vi
 			NewAnimBlueprintClass->EvaluateGraphExposedInputs.AddDefaulted(StructEvalHandlers.Num());
 			Record.EvaluationHandlerIdx = NewAnimBlueprintClass->EvaluateGraphExposedInputs.Num() - 1;
 
-			// add instance to class:
-			if(Record.bServicesInstanceProperties)
-			{
-				CreateEvaluationHandlerInstance(VisualAnimNode, Record);
-			}
-			else
-			{
-				CreateEvaluationHandlerStruct(VisualAnimNode, Record);
-			}
+			CreateEvaluationHandler(VisualAnimNode, Record);
 
 			ValidEvaluationHandlerList.Add(Record);
 		}
diff --git a/Engine/Source/Editor/KismetCompiler/Private/AnimBlueprintCompiler.h b/Engine/Source/Editor/KismetCompiler/Private/AnimBlueprintCompiler.h
--- a/Engine/Source/Editor/KismetCompiler/Private/AnimBlueprintCompiler.h
+++ b/Engine/Source/Editor/KismetCompiler/Private/AnimBlueprintCompiler.h
@@ -161,6 +161,7 @@ protected:
 		bool Apply(UObject* Object);
 	};
 
+	/** BP execution handler for Anim node - possibly */
 	struct FEvaluationHandlerRecord
 	{
 	public:
@@ -171,6 +172,9 @@ protected:
 		// The specific evaluation handler inside the specified node
 		int32 EvaluationHandlerIdx;
 
+		// Whether or not our serviced properties are actually on the anim node 
+		bool bServicesNodeProperties;
+
 		// Whether or not our serviced properties are actually on the instance instead of the node
 		bool bServicesInstanceProperties;
 
@@ -185,6 +189,7 @@ protected:
 		FEvaluationHandlerRecord()
 			: NodeVariableProperty(nullptr)
 			, EvaluationHandlerIdx(INDEX_NONE)
+			, bServicesNodeProperties(false)
 			, bServicesInstanceProperties(false)
 			, HandlerFunctionName(NAME_None)
 		{}
@@ -298,8 +303,7 @@ private:
 	UK2Node_CallFunction* SpawnCallAnimInstanceFunction(UEdGraphNode* SourceNode, FName FunctionName);
 
 	// Creates an evaluation handler for an FExposedValue property in an animation node
-	void CreateEvaluationHandlerStruct(UAnimGraphNode_Base* VisualAnimNode, FEvaluationHandlerRecord& Record);
-	void CreateEvaluationHandlerInstance(UAnimGraphNode_Base* VisualAnimNode, FEvaluationHandlerRecord& Record);
+	void CreateEvaluationHandler(UAnimGraphNode_Base* VisualAnimNode, FEvaluationHandlerRecord& Record);
 
 	// Prunes any nodes that aren't reachable via a pose link
 	void PruneIsolatedAnimationNodes(const TArray<UAnimGraphNode_Base*>& RootSet, TArray<UAnimGraphNode_Base*>& GraphNodes);
diff --git a/Engine/Source/Editor/MovieSceneTools/Private/TrackEditors/TransformTrackEditor.cpp b/Engine/Source/Editor/MovieSceneTools/Private/TrackEditors/TransformTrackEditor.cpp
--- a/Engine/Source/Editor/MovieSceneTools/Private/TrackEditors/TransformTrackEditor.cpp
+++ b/Engine/Source/Editor/MovieSceneTools/Private/TrackEditors/TransformTrackEditor.cpp
@@ -195,12 +195,12 @@ void F3DTransformTrackEditor::BuildTrackContextMenu( FMenuBuilder& MenuBuilder,
 	}
 	UMovieScene3DTransformTrack* TransformTrack = Cast<UMovieScene3DTransformTrack>( Track );
 	MenuBuilder.AddMenuEntry(
-		NSLOCTEXT( "Sequencer", "PasteMatineeMoveTrack", "Paste Matinee Move Track" ),
-		NSLOCTEXT( "Sequencer", "PasteMatineeMoveTrackTooltip", "Pastes keys from a Matinee move track into this track." ),
+		NSLOCTEXT("Sequencer", "PasteMatineeMoveTrack", "Paste Matinee Move Track"),
+		NSLOCTEXT("Sequencer", "PasteMatineeMoveTrackTooltip", "Pastes keys from a Matinee move track into this track."),
 		FSlateIcon(),
 		FUIAction(
-			FExecuteAction::CreateStatic( &CopyInterpMoveTrack, GetSequencer().ToSharedRef(), MoveTrack, TransformTrack ),
-			FCanExecuteAction::CreateStatic( &CanCopyInterpMoveTrack, MoveTrack, TransformTrack ) ) );
+			FExecuteAction::CreateStatic(&CopyInterpMoveTrack, GetSequencer().ToSharedRef(), MoveTrack, TransformTrack),
+			FCanExecuteAction::CreateStatic(&CanCopyInterpMoveTrack, MoveTrack, TransformTrack)));
 
 	//		FCanExecuteAction::CreateLambda( [=]()->bool { return MoveTrack != nullptr && MoveTrack->GetNumKeys() > 0 && TransformTrack != nullptr; } ) ) );
 
@@ -389,6 +389,35 @@ void F3DTransformTrackEditor::OnPostPropertyChanged(UObject* InObject, FProperty
 	}
 }
 
+bool F3DTransformTrackEditor::CanAddTransformKeysForSelectedObjects() const
+{
+	// WASD hotkeys to fly the viewport can conflict with hotkeys for setting keyframes (ie. s). 
+// If the viewport is moving, disregard setting keyframes.
+	for (FLevelEditorViewportClient* LevelVC : GEditor->GetLevelViewportClients())
+	{
+		if (LevelVC && LevelVC->IsMovingCamera())
+		{
+			return false;
+		}
+	}
+	TArray<UObject*> SelectedObjects;
+	for (FSelectedEditableComponentIterator It(GEditor->GetSelectedEditableComponentIterator()); It; ++It)
+	{
+		USceneComponent* SceneComponent = Cast<USceneComponent>(*It);
+		if (SceneComponent)
+		{
+			return true;
+		}
+	}
+
+	if (SelectedObjects.Num() == 0)
+	{
+		USelection* CurrentSelection = GEditor->GetSelectedActors();
+		CurrentSelection->GetSelectedObjects(AActor::StaticClass(), SelectedObjects);
+	}
+	return SelectedObjects.Num() > 0;
+}
+
 void F3DTransformTrackEditor::OnAddTransformKeysForSelectedObjects( EMovieSceneTransformChannel Channel )
 {
 	// WASD hotkeys to fly the viewport can conflict with hotkeys for setting keyframes (ie. s). 
@@ -431,20 +460,22 @@ void F3DTransformTrackEditor::BindCommands(TSharedRef<FUICommandList> SequencerC
 	CommandBindings = MakeShared<FUICommandList>();
 	CommandBindings->MapAction(
 		Commands.AddTransformKey,
-		FExecuteAction::CreateSP( this, &F3DTransformTrackEditor::OnAddTransformKeysForSelectedObjects, EMovieSceneTransformChannel::All ) );
+		FExecuteAction::CreateSP(this, &F3DTransformTrackEditor::OnAddTransformKeysForSelectedObjects, EMovieSceneTransformChannel::All),
+		FCanExecuteAction::CreateSP(this, &F3DTransformTrackEditor::CanAddTransformKeysForSelectedObjects));
+		
 
 	CommandBindings->MapAction(
 		Commands.AddTranslationKey,
-		FExecuteAction::CreateSP( this, &F3DTransformTrackEditor::OnAddTransformKeysForSelectedObjects, EMovieSceneTransformChannel::Translation ) );
-
+		FExecuteAction::CreateSP( this, &F3DTransformTrackEditor::OnAddTransformKeysForSelectedObjects, EMovieSceneTransformChannel::Translation ),
+		FCanExecuteAction::CreateSP(this, &F3DTransformTrackEditor::CanAddTransformKeysForSelectedObjects));
 	CommandBindings->MapAction(
 		Commands.AddRotationKey,
-		FExecuteAction::CreateSP( this, &F3DTransformTrackEditor::OnAddTransformKeysForSelectedObjects, EMovieSceneTransformChannel::Rotation ) );
-
+		FExecuteAction::CreateSP( this, &F3DTransformTrackEditor::OnAddTransformKeysForSelectedObjects, EMovieSceneTransformChannel::Rotation ),
+		FCanExecuteAction::CreateSP(this, &F3DTransformTrackEditor::CanAddTransformKeysForSelectedObjects));
 	CommandBindings->MapAction(
 		Commands.AddScaleKey,
-		FExecuteAction::CreateSP( this, &F3DTransformTrackEditor::OnAddTransformKeysForSelectedObjects, EMovieSceneTransformChannel::Scale ) );
-
+		FExecuteAction::CreateSP( this, &F3DTransformTrackEditor::OnAddTransformKeysForSelectedObjects, EMovieSceneTransformChannel::Scale ),
+		FCanExecuteAction::CreateSP(this, &F3DTransformTrackEditor::CanAddTransformKeysForSelectedObjects));
 	Commands.BindingCount++;
 
 	// Add these bindings to Sequencer
diff --git a/Engine/Source/Editor/MovieSceneTools/Private/TrackEditors/TransformTrackEditor.h b/Engine/Source/Editor/MovieSceneTools/Private/TrackEditors/TransformTrackEditor.h
--- a/Engine/Source/Editor/MovieSceneTools/Private/TrackEditors/TransformTrackEditor.h
+++ b/Engine/Source/Editor/MovieSceneTools/Private/TrackEditors/TransformTrackEditor.h
@@ -113,6 +113,12 @@ private:
 	/** Delegate for camera button lock tooltip */
 	FText GetLockCameraToolTip(FGuid ObjectGuid) const; 
 
+
+	/** Whether or not we can add a transform key for a selected object 
+	* @return Returns true if we can.
+	**/
+	bool CanAddTransformKeysForSelectedObjects() const;
+
 	/** Generates transform keys based on the last transform, the current transform, and other options. 
 		One transform key is generated for each individual key to be added to the section. */
 	void GetTransformKeys( const TOptional<FTransformData>& LastTransform, const FTransformData& CurrentTransform, EMovieSceneTransformChannel ChannelsToKey, FGeneratedTrackKeys& OutGeneratedKeys );
diff --git a/Engine/Source/Editor/MovieSceneTools/Public/KeyframeTrackEditor.h b/Engine/Source/Editor/MovieSceneTools/Public/KeyframeTrackEditor.h
--- a/Engine/Source/Editor/MovieSceneTools/Public/KeyframeTrackEditor.h
+++ b/Engine/Source/Editor/MovieSceneTools/Public/KeyframeTrackEditor.h
@@ -142,7 +142,7 @@ struct TAddKeyImpl<FMovieSceneFloatChannel, float> : IImpl
 		using namespace MovieScene;
 
 		FMovieSceneFloatChannel* Channel = Proxy.GetChannel<FMovieSceneFloatChannel>(ChannelIndex);
-		if (Channel && Channel->GetData().GetTimes().Num() == 0)
+		if (Channel && Channel->GetData().GetTimes().Num() == 0  && Channel->GetDefault() != ValueToSet)
 		{
 			if (Section->TryModify())
 			{
@@ -499,6 +499,7 @@ private:
 		return KeyPropertyResult;
 	}
 
+protected:
 	/* Returns whether a section was added */
 	FKeyPropertyResult AddKeysToSection(UMovieSceneSection* Section, FFrameNumber KeyTime, const FGeneratedTrackKeys& Keys, ESequencerKeyMode KeyMode)
 	{
diff --git a/Engine/Source/Editor/Persona/Persona.Build.cs b/Engine/Source/Editor/Persona/Persona.Build.cs
--- a/Engine/Source/Editor/Persona/Persona.Build.cs
+++ b/Engine/Source/Editor/Persona/Persona.Build.cs
@@ -30,6 +30,7 @@ public class Persona : ModuleRules
                 "AnimationEditor",
                 "MeshReductionInterface",
                 "SequenceRecorder",
+                "AnimationBlueprintEditor",
             }
         );
 
diff --git a/Engine/Source/Editor/Persona/Private/AnimGraphNodeDetails.cpp b/Engine/Source/Editor/Persona/Private/AnimGraphNodeDetails.cpp
--- a/Engine/Source/Editor/Persona/Private/AnimGraphNodeDetails.cpp
+++ b/Engine/Source/Editor/Persona/Private/AnimGraphNodeDetails.cpp
@@ -169,8 +169,6 @@ void FAnimGraphNodeDetails::CustomizeDetails(class IDetailLayoutBuilder& DetailB
 
 						ShowHidePropertyHandle->MarkHiddenByCustomization();
 
-						TSharedRef<SWidget> ShowHidePropertyWidget = CreateAsPinWidget(ShowHidePropertyHandle);
-
 						ValueWidget->SetVisibility(TAttribute<EVisibility>::Create(TAttribute<EVisibility>::FGetter::CreateSP(this, &FAnimGraphNodeDetails::GetVisibilityOfProperty, ShowHidePropertyHandle)));
 
 						// If we have an edit condition, that comes as part of the default name widget, so just use a text block to avoid duplicate checkboxes
@@ -187,26 +185,7 @@ void FAnimGraphNodeDetails::CustomizeDetails(class IDetailLayoutBuilder& DetailB
 							PropertyNameWidget = NameWidget;
 						}
 
-						NameWidget = 
-							SNew(SHorizontalBox)
-							+SHorizontalBox::Slot()
-							.HAlign(HAlign_Left)
-							.VAlign(VAlign_Center)
-							.FillWidth(1)
-							[
-								SNew(SBox)
-								.Clipping(EWidgetClipping::ClipToBounds)
-								[
-									PropertyNameWidget.ToSharedRef()
-								]
-							]
-							+SHorizontalBox::Slot()
-							.AutoWidth()
-							.HAlign(HAlign_Right)
-							.VAlign(VAlign_Center)
-							[
-								ShowHidePropertyWidget
-							];
+						NameWidget = PropertyNameWidget;
 
 						// we only show children if visibility is one
 						// whenever toggles, this gets called, so it will be refreshed
@@ -250,68 +229,6 @@ void FAnimGraphNodeDetails::CustomizeDetails(class IDetailLayoutBuilder& DetailB
 	}
 }
 
-TSharedRef<SWidget> FAnimGraphNodeDetails::CreateAsPinWidget(TSharedRef<IPropertyHandle> InPropertyHandle)
-{
-	TWeakPtr<IPropertyHandle> WeakPropertyHandle = InPropertyHandle;
-
-	auto IsCheckedLambda = [WeakPropertyHandle]()
-	{
-		if(WeakPropertyHandle.IsValid())
-		{
-			bool bValue;
-			FPropertyAccess::Result Result = WeakPropertyHandle.Pin()->GetValue(bValue);
-			if(Result == FPropertyAccess::MultipleValues)
-			{
-				return ECheckBoxState::Undetermined;
-			}
-			else
-			{
-				return bValue ? ECheckBoxState::Checked : ECheckBoxState::Unchecked;
-			}
-		}
-
-		return ECheckBoxState::Unchecked;
-	};
-
-	auto OnCheckStateChangedLambda = [WeakPropertyHandle](ECheckBoxState InCheckBoxState)
-	{
-		if(WeakPropertyHandle.IsValid())
-		{
-			bool bValue = InCheckBoxState == ECheckBoxState::Checked;
-			WeakPropertyHandle.Pin()->SetValue(bValue);
-		}
-	};
-
-	auto ImageLambda = [WeakPropertyHandle]()
-	{
-		if(WeakPropertyHandle.IsValid())
-		{
-			bool bValue;
-			FPropertyAccess::Result Result = WeakPropertyHandle.Pin()->GetValue(bValue);
-			if(Result == FPropertyAccess::MultipleValues)
-			{
-				return FEditorStyle::GetBrush("Kismet.VariableList.HideForInstance");
-			}
-			else
-			{
-				return bValue ? FEditorStyle::GetBrush("Kismet.VariableList.ExposeForInstance") : FEditorStyle::GetBrush("Kismet.VariableList.HideForInstance");
-			}
-		}
-
-		return FEditorStyle::GetBrush("Kismet.VariableList.HideForInstance");
-	};
-
-	return SNew(SCheckBox)
-		.ToolTipText(LOCTEXT("AsPinTooltip", "Show/hide this property as a pin on the node"))
-		.IsChecked_Lambda(IsCheckedLambda)
-		.OnCheckStateChanged_Lambda(OnCheckStateChangedLambda)
-		.Style(FEditorStyle::Get(), "CheckboxLookToggleButtonCheckbox")
-		[
-			SNew(SImage)
-			.Image_Lambda(ImageLambda)
-			.ColorAndOpacity(FLinearColor::Black)
-		];
-}
 
 TSharedRef<SWidget> FAnimGraphNodeDetails::CreatePropertyWidget(UProperty* TargetProperty, TSharedRef<IPropertyHandle> TargetPropertyHandle, UClass* NodeClass)
 {
@@ -711,10 +628,10 @@ TSharedRef<IPropertyTypeCustomization> FBoneSocketTargetCustomization::MakeInsta
 void FBoneSocketTargetCustomization::CustomizeChildren(TSharedRef<IPropertyHandle> StructPropertyHandle, class IDetailChildrenBuilder& ChildBuilder, IPropertyTypeCustomizationUtils& StructCustomizationUtils)
 {
 	// set property handle 
- 	SetPropertyHandle(StructPropertyHandle);
- 	// set editable skeleton info from struct
- 	SetEditableSkeleton(StructPropertyHandle);
- 	Build(StructPropertyHandle, ChildBuilder);
+	SetPropertyHandle(StructPropertyHandle);
+	// set editable skeleton info from struct
+	SetEditableSkeleton(StructPropertyHandle);
+	Build(StructPropertyHandle, ChildBuilder);
 }
 
 void FBoneSocketTargetCustomization::SetPropertyHandle(TSharedRef<IPropertyHandle> StructPropertyHandle)
@@ -1225,4 +1142,135 @@ void FPlayerTreeViewEntry::GenerateNameWidget(TSharedPtr<SHorizontalBox> Box)
 		];
 }
 
+void FAnimGraphNodeShowAsPinExtension::GetOptionalPinData(const IPropertyHandle& PropertyHandle, int32& OutOptionalPinIndex, UAnimGraphNode_Base*& OutAnimGraphNode) const
+{
+	OutOptionalPinIndex = INDEX_NONE;
+
+	TArray<UObject*> Objects;
+	PropertyHandle.GetOuterObjects(Objects);
+
+	UProperty* Property = PropertyHandle.GetProperty();
+	if(Property)
+	{
+		OutAnimGraphNode = Cast<UAnimGraphNode_Base>(Objects[0]);
+		if (OutAnimGraphNode != nullptr)
+		{
+			OutOptionalPinIndex = OutAnimGraphNode->ShowPinForProperties.IndexOfByPredicate([Property](const FOptionalPinFromProperty& InOptionalPin)
+			{
+				return Property->GetFName() == InOptionalPin.PropertyName;
+			});
+		}
+	}
+}
+
+bool FAnimGraphNodeShowAsPinExtension::IsPropertyExtendable(const UClass* InObjectClass, const IPropertyHandle& PropertyHandle) const
+{
+	int32 OptionalPinIndex;
+	UAnimGraphNode_Base* AnimGraphNode;
+	GetOptionalPinData(PropertyHandle, OptionalPinIndex, AnimGraphNode);
+
+	if(OptionalPinIndex != INDEX_NONE)
+	{
+		const FOptionalPinFromProperty& OptionalPin = AnimGraphNode->ShowPinForProperties[OptionalPinIndex];
+
+		// Not optional
+		if (!OptionalPin.bCanToggleVisibility && OptionalPin.bShowPin)
+		{
+			return false;
+		}
+
+		if(!PropertyHandle.GetProperty())
+		{
+			return false;
+		}
+
+		return OptionalPin.bCanToggleVisibility;
+	}
+
+	return false;
+}
+
+class SShowAsWidget : public SCompoundWidget
+{
+	SLATE_BEGIN_ARGS(SShowAsWidget) {}
+
+	SLATE_END_ARGS()
+
+	void Construct(const FArguments& InArgs, const TSharedRef<IPropertyHandle>& InPropertyHandle)
+	{
+		PropertyHandle = InPropertyHandle;
+
+		TSharedRef<SHorizontalBox> HorizontalBox = 
+			SNew(SHorizontalBox)
+			.ToolTipText(LOCTEXT("AsPinTooltip", "Show/hide this property as a pin on the node"));
+
+		TWeakPtr<SWidget> WeakHorizontalBox = HorizontalBox;
+
+		HorizontalBox->AddSlot()
+			.AutoWidth()
+			.VAlign(VAlign_Center)
+			[
+				SNew(STextBlock)
+				.Text(LOCTEXT("ExposeAsPinLabel", "Expose"))
+				.Font(IDetailLayoutBuilder::GetDetailFont())
+				.Visibility_Lambda([WeakHorizontalBox](){ return WeakHorizontalBox.IsValid() && WeakHorizontalBox.Pin()->IsHovered() ? EVisibility::Visible : EVisibility::Collapsed; })
+			];
+
+		HorizontalBox->AddSlot()
+			.AutoWidth()
+			.VAlign(VAlign_Center)
+			.Padding(3.0f, 0.0f, 0.0f, 0.0f)
+			[
+				SNew(SCheckBox)
+				.IsChecked(this, &SShowAsWidget::IsChecked)
+				.OnCheckStateChanged(this, &SShowAsWidget::OnCheckStateChanged)
+			];
+
+		ChildSlot
+		[
+			HorizontalBox
+		];
+	}
+
+	ECheckBoxState IsChecked() const
+	{
+		bool bValue;
+		FPropertyAccess::Result Result = PropertyHandle->GetValue(bValue);
+		if(Result == FPropertyAccess::MultipleValues)
+		{
+			return ECheckBoxState::Undetermined;
+		}
+		else
+		{
+			return bValue ? ECheckBoxState::Checked : ECheckBoxState::Unchecked;
+		}
+
+		return ECheckBoxState::Unchecked;
+	}
+
+	void OnCheckStateChanged(ECheckBoxState InCheckBoxState)
+	{
+		bool bValue = InCheckBoxState == ECheckBoxState::Checked;
+		PropertyHandle->SetValue(bValue);
+	}
+
+	TSharedPtr<IPropertyHandle> PropertyHandle;
+};
+
+TSharedRef<SWidget> FAnimGraphNodeShowAsPinExtension::GenerateExtensionWidget(const IDetailLayoutBuilder& InDetailBuilder, const UClass* InObjectClass, TSharedPtr<IPropertyHandle> PropertyHandle)
+{
+	int32 OptionalPinIndex;
+	UAnimGraphNode_Base* AnimGraphNode;
+	GetOptionalPinData(*PropertyHandle.Get(), OptionalPinIndex, AnimGraphNode);
+	check(OptionalPinIndex != INDEX_NONE);
+
+	const FName OptionalPinArrayEntryName(*FString::Printf(TEXT("ShowPinForProperties[%d].bShowPin"), OptionalPinIndex));
+	TSharedRef<IPropertyHandle> ShowHidePropertyHandle = InDetailBuilder.GetProperty(OptionalPinArrayEntryName, UAnimGraphNode_Base::StaticClass());
+
+	ShowHidePropertyHandle->MarkHiddenByCustomization();
+
+	return SNew(SShowAsWidget, ShowHidePropertyHandle);
+}
+
 #undef LOCTEXT_NAMESPACE
+
diff --git a/Engine/Source/Editor/Persona/Private/AnimGraphNodeDetails.h b/Engine/Source/Editor/Persona/Private/AnimGraphNodeDetails.h
--- a/Engine/Source/Editor/Persona/Private/AnimGraphNodeDetails.h
+++ b/Engine/Source/Editor/Persona/Private/AnimGraphNodeDetails.h
@@ -12,6 +12,7 @@
 #include "IDetailCustomization.h"
 #include "Widgets/Views/STableViewBase.h"
 #include "Widgets/Views/STableRow.h"
+#include "IDetailPropertyExtensionHandler.h"
 
 struct FAssetData;
 class FBlueprintEditor;
@@ -23,6 +24,18 @@ class USkeleton;
 class IEditableSkeleton;
 struct FAnimParentNodeAssetOverride;
 
+class FAnimGraphNodeShowAsPinExtension : public IDetailPropertyExtensionHandler
+{
+public:
+	// IDetailPropertyExtensionHandler interface
+	virtual bool IsPropertyExtendable(const UClass* InObjectClass, const IPropertyHandle& PropertyHandle) const override;
+	virtual TSharedRef<SWidget> GenerateExtensionWidget(const IDetailLayoutBuilder& InDetailBuilder, const UClass* InObjectClass, TSharedPtr<IPropertyHandle> PropertyHandle) override;
+
+private:
+	// Helper function
+	void GetOptionalPinData(const IPropertyHandle& PropertyHandle, int32& OutOptionalPinIndex, UAnimGraphNode_Base*& OutAnimGraphNode) const;
+};
+
 /////////////////////////////////////////////////////
 // FAnimGraphNodeDetails 
 
diff --git a/Engine/Source/Editor/Persona/Private/PersonaModule.cpp b/Engine/Source/Editor/Persona/Private/PersonaModule.cpp
--- a/Engine/Source/Editor/Persona/Private/PersonaModule.cpp
+++ b/Engine/Source/Editor/Persona/Private/PersonaModule.cpp
@@ -748,10 +748,12 @@ void FPersonaModule::AddLoopingInterpolation(TArray<TWeakObjectPtr<UAnimSequence
 	}
 }
 
-void FPersonaModule::CustomizeSlotNodeDetails(const TSharedRef<class IDetailsView>& InDetailsView, FOnInvokeTab InOnInvokeTab)
+void FPersonaModule::CustomizeBlueprintEditorDetails(const TSharedRef<class IDetailsView>& InDetailsView, FOnInvokeTab InOnInvokeTab)
 {
 	InDetailsView->RegisterInstancedCustomPropertyLayout(UAnimGraphNode_Slot::StaticClass(),
 		FOnGetDetailCustomizationInstance::CreateStatic(&FAnimGraphNodeSlotDetails::MakeInstance, InOnInvokeTab));
+
+	InDetailsView->SetExtensionHandler(MakeShared<FAnimGraphNodeShowAsPinExtension>());
 }
 
 IPersonaEditorModeManager* FPersonaModule::CreatePersonaEditorModeManager()
diff --git a/Engine/Source/Editor/Persona/Private/SAnimCurvePicker.cpp b/Engine/Source/Editor/Persona/Private/SAnimCurvePicker.cpp
new file mode 100644
--- /dev/null
+++ b/Engine/Source/Editor/Persona/Private/SAnimCurvePicker.cpp
@@ -0,0 +1,125 @@
+// Copyright 1998-2019 Epic Games, Inc. All Rights Reserved.
+
+#include "SAnimCurvePicker.h"
+#include "AssetRegistryModule.h"
+#include "IEditableSkeleton.h"
+#include "Animation/AnimationAsset.h"
+#include "Widgets/Views/SListView.h"
+#include "Widgets/Input/SSearchBox.h"
+#include "Widgets/SBoxPanel.h"
+#include "Widgets/Text/STextBlock.h"
+#include "Widgets/Layout/SBox.h"
+
+#define LOCTEXT_NAMESPACE "SAnimCurvePicker"
+
+void SAnimCurvePicker::Construct(const FArguments& InArgs, const TSharedRef<IEditableSkeleton>& InEditableSkeleton)
+{
+	OnCurveNamePicked = InArgs._OnCurveNamePicked;
+	EditableSkeleton = InEditableSkeleton;
+
+	ChildSlot
+	[
+		SNew(SVerticalBox)
+		+SVerticalBox::Slot()
+		.AutoHeight()
+		[
+			SNew(SSearchBox)
+			.HintText(LOCTEXT("SearchBoxHint", "Search Available Curves"))
+			.OnTextChanged(this, &SAnimCurvePicker::HandleFilterTextChanged)
+		]
+		+SVerticalBox::Slot()
+		.FillHeight(1.0f)
+		[
+			SAssignNew(NameListView, SListView<TSharedPtr<FString>>)
+			.SelectionMode(ESelectionMode::Single)
+			.ItemHeight(20.0f)
+			.ListItemsSource(&CurveNames)
+			.OnSelectionChanged(this, &SAnimCurvePicker::HandleSelectionChanged)
+			.OnGenerateRow(this, &SAnimCurvePicker::HandleGenerateRow)
+		]
+	];
+
+	RefreshListItems();
+}
+
+void SAnimCurvePicker::HandleSelectionChanged(TSharedPtr<FString> InItem, ESelectInfo::Type InSelectionType)
+{
+	OnCurveNamePicked.ExecuteIfBound(FName(**InItem.Get()));
+}
+
+TSharedRef<ITableRow> SAnimCurvePicker::HandleGenerateRow(TSharedPtr<FString> InItem, const TSharedRef<STableViewBase>& InOwnerTable)
+{
+	return 
+		SNew(STableRow<TSharedPtr<FString>>, InOwnerTable)
+		[
+			SNew(SBox)
+			.MinDesiredHeight(20.0f)
+			.VAlign(VAlign_Center)
+			[
+				SNew(STextBlock)
+				.Text(FText::FromString(*InItem.Get()))
+				.HighlightText_Lambda([this]() { return FText::FromString(FilterText); })
+			]
+		];
+}
+
+void SAnimCurvePicker::RefreshListItems()
+{
+	CurveNames.Empty();
+	UniqueCurveNames.Empty();
+
+	// We use the asset registry to query all assets with the supplied skeleton, and accumulate their curve names
+	FAssetRegistryModule& AssetRegistryModule = FModuleManager::LoadModuleChecked<FAssetRegistryModule>("AssetRegistry");
+
+	FARFilter Filter;
+	Filter.bRecursiveClasses = true;
+	Filter.ClassNames.Add(UAnimationAsset::StaticClass()->GetFName());
+	Filter.TagsAndValues.Add(TEXT("Skeleton"), FAssetData(&EditableSkeleton.Pin()->GetSkeleton()).GetExportTextName());
+
+	TArray<FAssetData> FoundAssetData;
+	AssetRegistryModule.Get().GetAssets(Filter, FoundAssetData);
+
+	// now build set of unique curve names
+	for (FAssetData& AssetData : FoundAssetData)
+	{
+		const FString TagValue = AssetData.GetTagValueRef<FString>(USkeleton::CurveNameTag);
+		if (!TagValue.IsEmpty())
+		{
+			// parse notifies
+			TArray<FString> CurveValues;
+			if (TagValue.ParseIntoArray(CurveValues, *USkeleton::CurveTagDelimiter, true) > 0)
+			{
+				for (const FString& CurveValue : CurveValues)
+				{
+					UniqueCurveNames.Add(CurveValue);
+				}
+			}
+		}
+	}
+
+	FilterAvailableCurves();
+}
+
+void SAnimCurvePicker::FilterAvailableCurves()
+{
+	CurveNames.Empty();
+
+	for (const FString& UniqueCurveName : UniqueCurveNames)
+	{
+		if (FilterText.IsEmpty() || UniqueCurveName.Contains(FilterText))
+		{
+			CurveNames.Add(MakeShared<FString>(UniqueCurveName));
+		}
+	}
+
+	NameListView->RequestListRefresh();
+}
+
+void SAnimCurvePicker::HandleFilterTextChanged(const FText& InFilterText)
+{
+	FilterText = InFilterText.ToString();
+
+	FilterAvailableCurves();
+}
+
+#undef LOCTEXT_NAMESPACE
\ No newline at end of file
diff --git a/Engine/Source/Editor/Persona/Private/SAnimCurvePicker.h b/Engine/Source/Editor/Persona/Private/SAnimCurvePicker.h
new file mode 100644
--- /dev/null
+++ b/Engine/Source/Editor/Persona/Private/SAnimCurvePicker.h
@@ -0,0 +1,55 @@
+// Copyright 1998-2019 Epic Games, Inc. All Rights Reserved.
+
+#pragma once
+
+#include "CoreMinimal.h"
+#include "Widgets/DeclarativeSyntaxSupport.h"
+#include "Widgets/SCompoundWidget.h"
+#include "Widgets/Views/STableRow.h"
+
+class IEditableSkeleton;
+
+DECLARE_DELEGATE_OneParam(FOnCurveNamePicked, const FName& /*PickedName*/)
+
+class SAnimCurvePicker : public SCompoundWidget
+{
+public:
+	SLATE_BEGIN_ARGS(SAnimCurvePicker) {}
+
+	SLATE_EVENT(FOnCurveNamePicked, OnCurveNamePicked)
+
+	SLATE_END_ARGS()
+
+	void Construct(const FArguments& InArgs, const TSharedRef<IEditableSkeleton>& InEditableSkeleton);
+
+private:
+	/** Refresh the list of available curves */
+	void RefreshListItems();
+
+	/** Filter available curves */
+	void FilterAvailableCurves();
+
+	/** UI handlers */
+	void HandleSelectionChanged(TSharedPtr<FString> InItem, ESelectInfo::Type InSelectionType);
+	TSharedRef<ITableRow> HandleGenerateRow(TSharedPtr<FString> InItem, const TSharedRef<STableViewBase>& InOwnerTable);
+	void HandleFilterTextChanged(const FText& InFilterText);
+
+private:
+	/** Delegate fired when a curve name is picked */
+	FOnCurveNamePicked OnCurveNamePicked;
+
+	/** The editable skeleton we use to grab curves from */
+	TWeakPtr<IEditableSkeleton> EditableSkeleton;
+
+	/** The names of the curves we are displaying */
+	TArray<TSharedPtr<FString>> CurveNames;
+
+	/** All the unique curve names we can find */
+	TSet<FString> UniqueCurveNames;
+
+	/** The string we use to filter curve names */
+	FString FilterText;
+
+	/** The list view used to display names */
+	TSharedPtr<SListView<TSharedPtr<FString>>> NameListView;
+};
\ No newline at end of file
diff --git a/Engine/Source/Editor/Persona/Private/SAnimNotifyPanel.cpp b/Engine/Source/Editor/Persona/Private/SAnimNotifyPanel.cpp
--- a/Engine/Source/Editor/Persona/Private/SAnimNotifyPanel.cpp
+++ b/Engine/Source/Editor/Persona/Private/SAnimNotifyPanel.cpp
@@ -44,6 +44,8 @@
 #include "ClassViewerModule.h"
 #include "SSkeletonAnimNotifies.h"
 #include "Subsystems/AssetEditorSubsystem.h"
+#include "IAnimationEditor.h"
+#include "IAnimationSequenceBrowser.h"
 
 // Track Panel drawing
 const float NotificationTrackHeight = 20.0f;
@@ -823,6 +825,9 @@ protected:
 	/** Opens the supplied blueprint in an editor */
 	void OnOpenNotifySource(UBlueprint* InSourceBlueprint) const;
 
+	/** Filters the asset browser by the selected notify */
+	void OnFilterSkeletonNotify(FName InName);
+
 	/**
 	 * Selects a node on the track. Supports multi selection
 	 * @param TrackNodeIndex - Index of the node to select.
@@ -3098,19 +3103,29 @@ TSharedPtr<SWidget> SAnimNotifyTrack::SummonContextMenu(const FGeometry& MyGeome
 		UObject* NotifyObject = NotifyEvent->Notify;
 		NotifyObject = NotifyObject ? NotifyObject : NotifyEvent->NotifyStateClass;
 
-		if (NotifyObject && Cast<UBlueprintGeneratedClass>(NotifyObject->GetClass()))
+		MenuBuilder.BeginSection("ViewSource", LOCTEXT("NotifyViewHeading", "View"));
+
+		if (NotifyObject)
 		{
-			if (UBlueprint* Blueprint = Cast<UBlueprint>(NotifyObject->GetClass()->ClassGeneratedBy))
+			if (Cast<UBlueprintGeneratedClass>(NotifyObject->GetClass()))
 			{
-				MenuBuilder.BeginSection("ViewSource", LOCTEXT("NotifyViewHeading", "View"));
-
-				NewAction.ExecuteAction.BindRaw(
-					this, &SAnimNotifyTrack::OnOpenNotifySource, Blueprint);
-				MenuBuilder.AddMenuEntry(LOCTEXT("OpenNotifyBlueprint", "Open Notify Blueprint"), LOCTEXT("OpenNotifyBlueprintTooltip", "Opens the source blueprint for this notify"), FSlateIcon(), NewAction);
-
-				MenuBuilder.EndSection(); //ViewSource
+				if (UBlueprint * Blueprint = Cast<UBlueprint>(NotifyObject->GetClass()->ClassGeneratedBy))
+				{
+					NewAction.ExecuteAction.BindRaw(
+						this, &SAnimNotifyTrack::OnOpenNotifySource, Blueprint);
+					MenuBuilder.AddMenuEntry(LOCTEXT("OpenNotifyBlueprint", "Open Notify Blueprint"), LOCTEXT("OpenNotifyBlueprintTooltip", "Opens the source blueprint for this notify"), FSlateIcon(), NewAction);
+				}
 			}
 		}
+		else
+		{
+			// skeleton notify
+			NewAction.ExecuteAction.BindRaw(
+				this, &SAnimNotifyTrack::OnFilterSkeletonNotify, NotifyEvent->NotifyName);
+			MenuBuilder.AddMenuEntry(LOCTEXT("FindNotifyReferences", "Find References"), LOCTEXT("FindNotifyReferencesTooltip", "Find all references to this skeleton notify in the asset browser"), FSlateIcon(), NewAction);
+		}
+
+		MenuBuilder.EndSection(); //ViewSource
 	}
 
 	FWidgetPath WidgetPath = MouseEvent.GetEventPath() != nullptr ? *MouseEvent.GetEventPath() : FWidgetPath();
@@ -3169,6 +3184,17 @@ void SAnimNotifyTrack::OnOpenNotifySource(UBlueprint* InSourceBlueprint) const
 	GEditor->GetEditorSubsystem<UAssetEditorSubsystem>()->OpenEditorForAsset(InSourceBlueprint);
 }
 
+void SAnimNotifyTrack::OnFilterSkeletonNotify(FName InName)
+{
+	// Open asset browser first
+	OnInvokeTab.ExecuteIfBound(FPersonaTabs::AssetBrowserID);
+
+	IAssetEditorInstance* AssetEditor = GEditor->GetEditorSubsystem<UAssetEditorSubsystem>()->FindEditorForAsset(Sequence, true);
+	check(AssetEditor->GetEditorName() == TEXT("AnimationEditor"));
+	IAnimationEditor* AnimationEditor = static_cast<IAnimationEditor*>(AssetEditor);
+	AnimationEditor->GetAssetBrowser()->FilterBySkeletonNotify(InName);
+}
+
 bool SAnimNotifyTrack::IsSingleNodeSelected()
 {
 	return SelectedNodeIndices.Num() == 1;
diff --git a/Engine/Source/Editor/Persona/Private/SAnimationBlendSpaceGridWidget.cpp b/Engine/Source/Editor/Persona/Private/SAnimationBlendSpaceGridWidget.cpp
--- a/Engine/Source/Editor/Persona/Private/SAnimationBlendSpaceGridWidget.cpp
+++ b/Engine/Source/Editor/Persona/Private/SAnimationBlendSpaceGridWidget.cpp
@@ -82,6 +82,10 @@ void SBlendSpaceGridWidget::Construct(const FArguments& InArgs)
 	// Retrieve background and sample key brushes 
 	BackgroundImage = FEditorStyle::GetBrush(TEXT("Graph.Panel.SolidBackground"));
 	KeyBrush = FEditorStyle::GetBrush("CurveEd.CurveKey");
+	ArrowBrushes[(uint8)EArrowDirection::Up] = FEditorStyle::GetBrush("BlendSpaceEditor.ArrowUp");
+	ArrowBrushes[(uint8)EArrowDirection::Down] = FEditorStyle::GetBrush("BlendSpaceEditor.ArrowDown");
+	ArrowBrushes[(uint8)EArrowDirection::Right] = FEditorStyle::GetBrush("BlendSpaceEditor.ArrowRight");
+	ArrowBrushes[(uint8)EArrowDirection::Left] = FEditorStyle::GetBrush("BlendSpaceEditor.ArrowLeft");
 
 	// Retrieve font data 
 	FontInfo = FEditorStyle::GetFontStyle("CurveEd.InfoFont");
@@ -412,8 +416,19 @@ void SBlendSpaceGridWidget::PaintAxisText(const FGeometry& AllottedGeometry, con
 	FString Text = ParameterXName.ToString();
 	FVector2D TextSize = FontMeasure->Measure(Text, FontInfo);
 
+	// arrow left
+	FVector2D TextPosition = FVector2D(GridCenter.X - (TextSize.X * .5f), CachedGridRectangle.Bottom + TextMargin + (KeySize.Y * .25f));
+	FVector2D ArrowSize = ArrowBrushes[(uint8)EArrowDirection::Left]->GetImageSize();
+	FVector2D ArrowPosition = FVector2D(TextPosition.X - ArrowSize.X - 10.f/* give padding*/, TextPosition.Y);
+	FSlateDrawElement::MakeBox(OutDrawElements, DrawLayerId + 1, AllottedGeometry.ToPaintGeometry(ArrowPosition, ArrowSize), ArrowBrushes[(uint8)EArrowDirection::Left], ESlateDrawEffect::None, FLinearColor::White);
+
 	// Label
-	FSlateDrawElement::MakeText(OutDrawElements, DrawLayerId + 1, AllottedGeometry.MakeChild(FVector2D(GridCenter.X - (TextSize.X * .5f), CachedGridRectangle.Bottom + TextMargin + (KeySize.Y * .25f)), FVector2D(1.0f, 1.0f)).ToPaintGeometry(), Text, FontInfo, ESlateDrawEffect::None, FLinearColor::White);
+	FSlateDrawElement::MakeText(OutDrawElements, DrawLayerId + 1, AllottedGeometry.MakeChild(TextPosition, FVector2D(1.0f, 1.0f)).ToPaintGeometry(), Text, FontInfo, ESlateDrawEffect::None, FLinearColor::White);
+
+	// arrow right
+	ArrowSize = ArrowBrushes[(uint8)EArrowDirection::Right]->GetImageSize();
+	ArrowPosition = FVector2D(TextPosition.X + TextSize.X + 10.f/* give padding*/, TextPosition.Y);
+	FSlateDrawElement::MakeBox(OutDrawElements, DrawLayerId + 1, AllottedGeometry.ToPaintGeometry(ArrowPosition, ArrowSize), ArrowBrushes[(uint8)EArrowDirection::Right], ESlateDrawEffect::None, FLinearColor::White);
 
 	Text = FString::SanitizeFloat(SampleValueMin.X);
 	TextSize = FontMeasure->Measure(Text, FontInfo);
@@ -434,8 +449,19 @@ void SBlendSpaceGridWidget::PaintAxisText(const FGeometry& AllottedGeometry, con
 		Text = ParameterYName.ToString();
 		TextSize = FontMeasure->Measure(Text, FontInfo);
 
+		// arrow up
+		TextPosition = FVector2D(((GridMargin.Left - TextSize.X) * 0.5f - (KeySize.X * .25f)) + GridRatioMargin.Left, GridCenter.Y - (TextSize.Y * .5f));
+		ArrowSize = ArrowBrushes[(uint8)EArrowDirection::Up]->GetImageSize();
+		ArrowPosition = FVector2D(TextPosition.X + TextSize.X * 0.5f - ArrowSize.X * 0.5f, TextPosition.Y - ArrowSize.X - 10.f/* give padding*/);
+		FSlateDrawElement::MakeBox(OutDrawElements, DrawLayerId + 1, AllottedGeometry.ToPaintGeometry(ArrowPosition, ArrowSize), ArrowBrushes[(uint8)EArrowDirection::Up], ESlateDrawEffect::None, FLinearColor::White);
+
 		// Label
-		FSlateDrawElement::MakeText(OutDrawElements, DrawLayerId + 1, AllottedGeometry.MakeChild(FVector2D(((GridMargin.Left - TextSize.X) * 0.5f - (KeySize.X * .25f)) + GridRatioMargin.Left, GridCenter.Y - (TextSize.Y * .5f)), FVector2D(1.0f, 1.0f)).ToPaintGeometry(), Text,	FontInfo, ESlateDrawEffect::None, FLinearColor::White);
+		FSlateDrawElement::MakeText(OutDrawElements, DrawLayerId + 1, AllottedGeometry.MakeChild(TextPosition, FVector2D(1.0f, 1.0f)).ToPaintGeometry(), Text,	FontInfo, ESlateDrawEffect::None, FLinearColor::White);
+
+		// arrow down
+		ArrowSize = ArrowBrushes[(uint8)EArrowDirection::Down]->GetImageSize();
+		ArrowPosition = FVector2D(TextPosition.X + TextSize.X * 0.5f - ArrowSize.X * 0.5f, TextPosition.Y + TextSize.Y + 10.f/* give padding*/);
+		FSlateDrawElement::MakeBox(OutDrawElements, DrawLayerId + 1, AllottedGeometry.ToPaintGeometry(ArrowPosition, ArrowSize), ArrowBrushes[(uint8)EArrowDirection::Down], ESlateDrawEffect::None, FLinearColor::White);
 
 		Text = FString::SanitizeFloat(SampleValueMin.Y);
 		TextSize = FontMeasure->Measure(Text, FontInfo);
diff --git a/Engine/Source/Editor/Persona/Private/SAnimationBlendSpaceGridWidget.h b/Engine/Source/Editor/Persona/Private/SAnimationBlendSpaceGridWidget.h
--- a/Engine/Source/Editor/Persona/Private/SAnimationBlendSpaceGridWidget.h
+++ b/Engine/Source/Editor/Persona/Private/SAnimationBlendSpaceGridWidget.h
@@ -269,6 +269,15 @@ private:
 	FVector2D KeySize;
 	const FSlateBrush* KeyBrush;
 	const FSlateBrush* BackgroundImage;
+	enum class EArrowDirection : uint8 
+	{
+		Left, 
+		Right, 
+		Up, 
+		Down, 
+		Max
+	};
+	const FSlateBrush* ArrowBrushes[(int32)EArrowDirection::Max];
 
 	/** Individual sample state colours */
 	FSlateColor HighlightKeyColor;
diff --git a/Engine/Source/Editor/Persona/Private/SAnimationSequenceBrowser.cpp b/Engine/Source/Editor/Persona/Private/SAnimationSequenceBrowser.cpp
--- a/Engine/Source/Editor/Persona/Private/SAnimationSequenceBrowser.cpp
+++ b/Engine/Source/Editor/Persona/Private/SAnimationSequenceBrowser.cpp
@@ -41,6 +41,10 @@
 #include "Preferences/PersonaOptions.h"
 #include "Settings/SkeletalMeshEditorSettings.h"
 #include "Subsystems/AssetEditorSubsystem.h"
+#include "SSkeletonAnimNotifies.h"
+#include "SAnimCurveViewer.h"
+#include "IEditableSkeleton.h"
+#include "SAnimCurvePicker.h"
 
 #define LOCTEXT_NAMESPACE "SequenceBrowser"
 
@@ -65,6 +69,199 @@ public:
 	}
 };
 
+/** A filter that displays animations that use a skeleton notify */
+class FFrontendFilter_SkeletonNotify : public FFrontendFilter
+{
+public:
+	FFrontendFilter_SkeletonNotify(TSharedPtr<FFrontendFilterCategory> InCategory, TSharedPtr<IEditableSkeleton> InEditableSkeleton)
+		: FFrontendFilter(InCategory) 
+		, EditableSkeleton(InEditableSkeleton)
+	{}
+
+	// FFrontendFilter implementation
+	virtual FString GetName() const override 
+	{
+		return TEXT("SkeletonNotifyAssets"); 
+	}
+
+	virtual FText GetDisplayName() const override 
+	{ 
+		return NotifyString.IsEmpty() ? LOCTEXT("FFrontendFilter_SkeletonNotifyAssetsEmpty", "Uses Skeleton Notify...") : FText::Format(LOCTEXT("FFrontendFilter_SkeletonNotifyAssets", "Uses Skeleton Notify: {0}"), FText::FromString(NotifyString));
+	}
+
+	virtual FText GetToolTipText() const override
+	{
+		return NotifyString.IsEmpty() ? LOCTEXT("FFrontendFilter_SkeletonNotifyAssetsEmptyTooltip", "Show assets that use a specified skeleton notify") : FText::Format(LOCTEXT("FFrontendFilter_SkeletonNotifyAssetsTooltip", "Show assets that use skeleton notify '{0}'"), FText::FromString(NotifyString));
+	}
+
+	virtual void ModifyContextMenu(FMenuBuilder& MenuBuilder) override
+	{
+		MenuBuilder.BeginSection(TEXT("NotifySection"), LOCTEXT("NotifySectionHeading", "Choose Notify"));
+
+		TSharedRef<SWidget> Widget =
+			SNew(SBox)
+			.HeightOverride(300.0f)
+			.WidthOverride(200.0f)
+			[
+				SNew(SSkeletonAnimNotifies, EditableSkeleton.ToSharedRef())
+				.IsPicker(true)
+				.OnItemSelected_Lambda([this](const FName& InNotifyName)
+				{
+					FSlateApplication::Get().DismissAllMenus();
+					NotifyString = InNotifyName.ToString();
+					OnChanged().Broadcast();
+				})
+			];
+
+		MenuBuilder.AddWidget(Widget, FText(), true);
+
+		MenuBuilder.EndSection();
+	}
+
+	virtual FLinearColor GetColor() const override
+	{
+		return FLinearColor(0.0f, 0.4f, 0.6f, 1);
+	}
+
+	// IFilter implementation
+	virtual bool PassesFilter(FAssetFilterType InItem) const override
+	{
+		if (!NotifyString.IsEmpty())
+		{
+			const FString TagValue = InItem.GetTagValueRef<FString>(USkeleton::AnimNotifyTag);
+			if (!TagValue.IsEmpty())
+			{
+				// parse notifies
+				TArray<FString> NotifyValues;
+				if (TagValue.ParseIntoArray(NotifyValues, *USkeleton::AnimNotifyTagDelimiter, true) > 0)
+				{
+					for (const FString& NotifyValue : NotifyValues)
+					{
+						if (NotifyValue == NotifyString)
+						{
+							return true;
+						}
+					}
+				}
+			}
+
+			return false;
+		}
+
+		return true;
+	}
+
+	void SetNotifyFilter(const FName& InNotifyFilter)
+	{
+		NotifyString = InNotifyFilter.ToString();
+		OnChanged().Broadcast();
+		SetActive(true);
+	}
+
+public:
+	/** Notify string to use when filtering */
+	FString NotifyString;
+
+	/** Editable skeleton used to access available notifies */
+	TSharedPtr<IEditableSkeleton> EditableSkeleton;
+};
+
+/** A filter that displays animations that use a curve */
+class FFrontendFilter_Curves : public FFrontendFilter
+{
+public:
+	FFrontendFilter_Curves(TSharedPtr<FFrontendFilterCategory> InCategory, TSharedPtr<IEditableSkeleton> InEditableSkeleton)
+		: FFrontendFilter(InCategory)
+		, EditableSkeleton(InEditableSkeleton)
+	{}
+
+	// FFrontendFilter implementation
+	virtual FString GetName() const override
+	{
+		return TEXT("CurveAssets");
+	}
+
+	virtual FText GetDisplayName() const override
+	{
+		return CurveString.IsEmpty() ? LOCTEXT("FFrontendFilter_CurveAssetsEmpty", "Uses Curve...") : FText::Format(LOCTEXT("FFrontendFilter_CurveAssets", "Uses Curve: {0}"), FText::FromString(CurveString));
+	}
+
+	virtual FText GetToolTipText() const override
+	{
+		return CurveString.IsEmpty() ? LOCTEXT("FFrontendFilter_CurveAssetsEmptyTooltip", "Show assets that use a specified curve") : FText::Format(LOCTEXT("FFrontendFilter_CurveAssetsTooltip", "Show assets that use curve '{0}'"), FText::FromString(CurveString));
+	}
+
+	virtual void ModifyContextMenu(FMenuBuilder& MenuBuilder) override
+	{
+		MenuBuilder.BeginSection(TEXT("CurveSection"), LOCTEXT("CurveSectionHeading", "Choose Curve"));
+
+		TSharedRef<SWidget> Widget = 
+			SNew(SBox)
+			.HeightOverride(300.0f)
+			.WidthOverride(200.0f)
+			[
+				SNew(SAnimCurvePicker, EditableSkeleton.ToSharedRef())
+				.OnCurveNamePicked_Lambda([this](const FName& InCurveName)
+				{
+					FSlateApplication::Get().DismissAllMenus();
+					CurveString = InCurveName.ToString();
+					OnChanged().Broadcast();
+				})
+			];
+
+		MenuBuilder.AddWidget(Widget, FText(), true);
+
+		MenuBuilder.EndSection();
+	}
+
+	virtual FLinearColor GetColor() const override
+	{
+		return FLinearColor(0.0f, 0.6f, 0.2f, 1);
+	}
+
+	// IFilter implementation
+	virtual bool PassesFilter(FAssetFilterType InItem) const override
+	{
+		if (!CurveString.IsEmpty())
+		{
+			const FString TagValue = InItem.GetTagValueRef<FString>(USkeleton::CurveNameTag);
+			if (!TagValue.IsEmpty())
+			{
+				// parse curves
+				TArray<FString> CurveValues;
+				if (TagValue.ParseIntoArray(CurveValues, *USkeleton::CurveTagDelimiter, true) > 0)
+				{
+					for (const FString& CurveValue : CurveValues)
+					{
+						if (CurveValue == CurveString)
+						{
+							return true;
+						}
+					}
+				}
+			}
+
+			return false;
+		}
+
+		return true;
+	}
+
+	void SetCurveFilter(const FName& InCurveFilter)
+	{
+		CurveString = InCurveFilter.ToString();
+		OnChanged().Broadcast();
+		SetActive(true);
+	}
+
+public:
+	/** Curve string to use when filtering */
+	FString CurveString;
+
+	/** Editable skeleton used to access available curves */
+	TSharedPtr<IEditableSkeleton> EditableSkeleton;
+};
+
 /** A filter that displays sound waves */
 class FFrontendFilter_SoundWaves : public FFrontendFilter
 {
@@ -635,6 +832,10 @@ void SAnimationSequenceBrowser::Construct(const FArguments& InArgs, const TShare
 		Config.ExtraFrontendFilters.Add(FolderFilter);
 	}
 
+	SkeletonNotifyFilter = MakeShareable(new FFrontendFilter_SkeletonNotify(AnimCategory, InPersonaToolkit->GetEditableSkeleton()));
+	Config.ExtraFrontendFilters.Add(SkeletonNotifyFilter.ToSharedRef());
+	Config.ExtraFrontendFilters.Add(MakeShareable(new FFrontendFilter_Curves(AnimCategory, InPersonaToolkit->GetEditableSkeleton())));
+
 	Config.OnIsAssetValidForCustomToolTip = FOnIsAssetValidForCustomToolTip::CreateLambda([](const FAssetData& AssetData) {return AssetData.IsAssetLoaded(); });
 	Config.OnGetCustomAssetToolTip = FOnGetCustomAssetToolTip::CreateSP(this, &SAnimationSequenceBrowser::CreateCustomAssetToolTip);
 	Config.OnVisualizeAssetToolTip = FOnVisualizeAssetToolTip::CreateSP(this, &SAnimationSequenceBrowser::OnVisualizeAssetToolTip);
@@ -956,6 +1157,13 @@ void SAnimationSequenceBrowser::SelectAsset(UAnimationAsset * AnimAsset)
 	}
 }
 
+void SAnimationSequenceBrowser::FilterBySkeletonNotify(const FName& InNotifyName)
+{
+	check(SkeletonNotifyFilter.IsValid());
+
+	SkeletonNotifyFilter->SetNotifyFilter(InNotifyName);
+}
+
 void SAnimationSequenceBrowser::AddToHistory(UAnimationAsset * AnimAsset)
 {
 	if (AnimAsset)
diff --git a/Engine/Source/Editor/Persona/Private/SAnimationSequenceBrowser.h b/Engine/Source/Editor/Persona/Private/SAnimationSequenceBrowser.h
--- a/Engine/Source/Editor/Persona/Private/SAnimationSequenceBrowser.h
+++ b/Engine/Source/Editor/Persona/Private/SAnimationSequenceBrowser.h
@@ -27,6 +27,7 @@ class UAnimationAsset;
 class UDebugSkelMeshComponent;
 class USoundWave;
 class FFrontendFilter_Folder;
+class FFrontendFilter_SkeletonNotify;
 
 //////////////////////////////////////////////////////////////////////////
 // FAnimationAssetViewportClient
@@ -67,7 +68,9 @@ public:
 
 	/** IAnimationSequenceBrowser interface */
 	virtual void SelectAsset(UAnimationAsset * AnimAsset) override;
-
+	virtual void AddToHistory(UAnimationAsset* AnimAsset) override;
+	virtual void FilterBySkeletonNotify(const FName& InNotifyName) override;
+	
 	/** Delegate that handles creation of context menu */
 	TSharedPtr<SWidget> OnGetAssetContextMenu(const TArray<FAssetData>& SelectedAssets);
 
@@ -104,9 +107,6 @@ public:
 	 */
 	void OnCreateCopy(TArray<FAssetData> Selected);
 
-	/** public reference to add to history */
-	virtual void AddToHistory(UAnimationAsset * AnimAsset) override;
-
 protected:
 	bool CanShowColumnForAssetRegistryTag(FName AssetType, FName TagName) const;
 	
@@ -262,6 +262,9 @@ protected:
 	/** All the folder filters we have */
 	TArray<TSharedPtr<FFrontendFilter_Folder>> FolderFilters;
 
+	/** The skeleton notify filter */
+	TSharedPtr<FFrontendFilter_SkeletonNotify> SkeletonNotifyFilter;
+
 	void RetargetAnimationHandler(USkeleton* OldSkeleton, USkeleton* NewSkeleton, bool bRemapReferencedAssets, bool bAllowRemapToExisting, bool bConvertSpaces, const EditorAnimUtils::FNameDuplicationRule* NameRule, TArray<TWeakObjectPtr<UObject>> InAnimAssets);
 
 private:
diff --git a/Engine/Source/Editor/Persona/Private/SSkeletonAnimNotifies.cpp b/Engine/Source/Editor/Persona/Private/SSkeletonAnimNotifies.cpp
--- a/Engine/Source/Editor/Persona/Private/SSkeletonAnimNotifies.cpp
+++ b/Engine/Source/Editor/Persona/Private/SSkeletonAnimNotifies.cpp
@@ -16,6 +16,9 @@
 #include "IEditableSkeleton.h"
 #include "TabSpawners.h"
 #include "Editor.h"
+#include "IAnimationEditor.h"
+#include "IAnimationBlueprintEditor.h"
+#include "IAnimationSequenceBrowser.h"
 
 #define LOCTEXT_NAMESPACE "SkeletonAnimNotifies"
 
@@ -44,14 +47,14 @@ FSkeletonAnimNotifiesSummoner::FSkeletonAnimNotifiesSummoner(TSharedPtr<class FA
 
 TSharedRef<SWidget> FSkeletonAnimNotifiesSummoner::CreateTabBody(const FWorkflowTabSpawnInfo& Info) const
 {
-	return SNew(SSkeletonAnimNotifies, EditableSkeleton.Pin().ToSharedRef())
+	return SNew(SSkeletonAnimNotifies, EditableSkeleton.Pin().ToSharedRef(), HostingApp.Pin())
 		.OnObjectsSelected(OnObjectsSelected);
 }
 
 /////////////////////////////////////////////////////
 // SSkeletonAnimNotifies
 
-void SSkeletonAnimNotifies::Construct(const FArguments& InArgs, const TSharedRef<IEditableSkeleton>& InEditableSkeleton)
+void SSkeletonAnimNotifies::Construct(const FArguments& InArgs, const TSharedRef<IEditableSkeleton>& InEditableSkeleton, const TSharedPtr<class FAssetEditorToolkit>& InHostingApp)
 {
 	OnObjectsSelected = InArgs._OnObjectsSelected;
 	OnItemSelected = InArgs._OnItemSelected;
@@ -59,6 +62,7 @@ void SSkeletonAnimNotifies::Construct(const FArguments& InArgs, const TSharedRef
 	bIsSyncMarker = InArgs._IsSyncMarker;
 
 	EditableSkeleton = InEditableSkeleton;
+	WeakHostingApp = InHostingApp;
 
 	if (bIsSyncMarker)
 	{
@@ -202,7 +206,14 @@ TSharedPtr<SWidget> SSkeletonAnimNotifies::OnGetContextMenuContent() const
 				MenuBuilder.AddMenuEntry(Label, ToolTipText, FSlateIcon(), Action);
 			}
 
-
+			if(WeakHostingApp.IsValid() && NotifiesListView->GetNumItemsSelected() == 1)
+			{
+				FUIAction Action = FUIAction(FExecuteAction::CreateSP(const_cast<SSkeletonAnimNotifies*>(this), &SSkeletonAnimNotifies::OnFindReferences),
+					FCanExecuteAction::CreateSP(this, &SSkeletonAnimNotifies::CanPerformDelete));
+				const FText Label = LOCTEXT("FindNotifyReferences", "Find References");
+				const FText ToolTipText = LOCTEXT("FindNotifyReferencesTooltip", "Find all references to this skeleton notify in the asset browser");
+				MenuBuilder.AddMenuEntry(Label, ToolTipText, FSlateIcon(), Action);
+			}
 		}
 		MenuBuilder.EndSection();
 	}
@@ -462,4 +473,30 @@ void SSkeletonAnimNotifies::NotifyUser( FNotificationInfo& NotificationInfo )
 		Notification->SetCompletionState( SNotificationItem::CS_Fail );
 	}
 }
+
+void SSkeletonAnimNotifies::OnFindReferences()
+{
+	TSharedPtr<FAssetEditorToolkit> HostingApp = WeakHostingApp.Pin();
+	if (HostingApp.IsValid())
+	{
+		check(NotifiesListView->GetNumItemsSelected() == 1);
+		TArray<TSharedPtr<FDisplayedAnimNotifyInfo>> SelectedItems;
+		NotifiesListView->GetSelectedItems(SelectedItems);
+		FName Name = SelectedItems[0]->Name;
+
+		HostingApp->InvokeTab(FPersonaTabs::AssetBrowserID);
+
+		if (HostingApp->GetEditorName() == TEXT("AnimationEditor"))
+		{
+			TSharedPtr<IAnimationEditor> AnimationEditor = StaticCastSharedPtr<IAnimationEditor>(HostingApp);
+			AnimationEditor->GetAssetBrowser()->FilterBySkeletonNotify(Name);
+		}
+		else if (HostingApp->GetEditorName() == TEXT("AnimationBlueprintEditor"))
+		{
+			TSharedPtr<IAnimationBlueprintEditor> AnimationBlueprintEditor = StaticCastSharedPtr<IAnimationBlueprintEditor>(HostingApp);
+			AnimationBlueprintEditor->GetAssetBrowser()->FilterBySkeletonNotify(Name);
+		}
+	}
+}
+
 #undef LOCTEXT_NAMESPACE
diff --git a/Engine/Source/Editor/Persona/Private/SSkeletonAnimNotifies.h b/Engine/Source/Editor/Persona/Private/SSkeletonAnimNotifies.h
--- a/Engine/Source/Editor/Persona/Private/SSkeletonAnimNotifies.h
+++ b/Engine/Source/Editor/Persona/Private/SSkeletonAnimNotifies.h
@@ -102,7 +102,7 @@ public:
 
 	SLATE_END_ARGS()
 public:
-	void Construct(const FArguments& InArgs, const TSharedRef<IEditableSkeleton>& InEditableSkeleton);
+	void Construct(const FArguments& InArgs, const TSharedRef<IEditableSkeleton>& InEditableSkeleton, const TSharedPtr<class FAssetEditorToolkit>& InHostingApp = nullptr);
 
 	~SSkeletonAnimNotifies();
 
@@ -187,6 +187,9 @@ private:
 	/** Handle when an item is scrolled into view, triggers a rename for new items */
 	void OnItemScrolledIntoView(TSharedPtr<FDisplayedAnimNotifyInfo> InItem, const TSharedPtr<ITableRow>& InTableRow);
 
+	/** Handle find references from the context menu */
+	void OnFindReferences();
+
 	/** The skeleton we are currently editing */
 	TSharedPtr<class IEditableSkeleton> EditableSkeleton;
 
@@ -211,6 +214,9 @@ private:
 	/** Delegate fired when a notify is selected */
 	FOnItemSelected OnItemSelected;
 
+	/** The app we are embedded in */
+	TWeakPtr<class FAssetEditorToolkit> WeakHostingApp;
+
 	/** Whether we should use this dialog as a picker or an editor. In picker mode we cant add, remove or rename notifies. */
 	bool bIsPicker;
 
diff --git a/Engine/Source/Editor/Persona/Public/IAnimationSequenceBrowser.h b/Engine/Source/Editor/Persona/Public/IAnimationSequenceBrowser.h
--- a/Engine/Source/Editor/Persona/Public/IAnimationSequenceBrowser.h
+++ b/Engine/Source/Editor/Persona/Public/IAnimationSequenceBrowser.h
@@ -12,6 +12,10 @@ class IAnimationSequenceBrowser : public SCompoundWidget
 public:
 	/** Select the specified asset */
 	virtual void SelectAsset(UAnimationAsset* InAsset) = 0;
+
 	/** Add to navigation history */
-	virtual void AddToHistory(UAnimationAsset * AnimAsset) = 0;
+	virtual void AddToHistory(UAnimationAsset* AnimAsset) = 0;
+
+	/** Add filter by skeleton notify name */
+	virtual void FilterBySkeletonNotify(const FName& InNotifyName) = 0;
 };
diff --git a/Engine/Source/Editor/Persona/Public/PersonaModule.h b/Engine/Source/Editor/Persona/Public/PersonaModule.h
--- a/Engine/Source/Editor/Persona/Public/PersonaModule.h
+++ b/Engine/Source/Editor/Persona/Public/PersonaModule.h
@@ -269,8 +269,11 @@ public:
 	/** Add looping interpolation to the list of animations */
 	virtual void AddLoopingInterpolation(TArray<TWeakObjectPtr<class UAnimSequence>>& AnimSequences);
 
+	UE_DEPRECATED(4.24, "Function renamed, please use CustomizeBlueprintEditorDetails")
+	virtual void CustomizeSlotNodeDetails(const TSharedRef<class IDetailsView>& InDetailsView, FOnInvokeTab InOnInvokeTab) { CustomizeBlueprintEditorDetails(InDetailsView, InOnInvokeTab); }
+
 	/** Customize the details of a slot node for the specified details view */
-	virtual void CustomizeSlotNodeDetails(const TSharedRef<class IDetailsView>& InDetailsView, FOnInvokeTab InOnInvokeTab);
+	virtual void CustomizeBlueprintEditorDetails(const TSharedRef<class IDetailsView>& InDetailsView, FOnInvokeTab InOnInvokeTab);
 
 	/** Create a Persona editor mode manager. Should be destroyed using plain ol' delete. Note: Only FPersonaEditMode-derived modes should be used with this manager! */
 	virtual class IPersonaEditorModeManager* CreatePersonaEditorModeManager();
diff --git a/Engine/Source/Editor/PropertyEditor/Private/DetailLayoutBuilderImpl.cpp b/Engine/Source/Editor/PropertyEditor/Private/DetailLayoutBuilderImpl.cpp
--- a/Engine/Source/Editor/PropertyEditor/Private/DetailLayoutBuilderImpl.cpp
+++ b/Engine/Source/Editor/PropertyEditor/Private/DetailLayoutBuilderImpl.cpp
@@ -152,7 +152,7 @@ IDetailPropertyRow* FDetailLayoutBuilderImpl::EditDefaultProperty(TSharedPtr<IPr
 	return nullptr;
 }
 
-TSharedRef<IPropertyHandle> FDetailLayoutBuilderImpl::GetProperty( const FName PropertyPath, const UStruct* ClassOutermost, FName InInstanceName )
+TSharedRef<IPropertyHandle> FDetailLayoutBuilderImpl::GetProperty( const FName PropertyPath, const UStruct* ClassOutermost, FName InInstanceName ) const
 {	
 	TSharedPtr<FPropertyHandleBase> PropertyHandle; 
 
@@ -551,7 +551,7 @@ TSharedPtr<FPropertyNode> FDetailLayoutBuilderImpl::GetPropertyNodeInternal( con
 }
 
 
-TSharedRef<IPropertyHandle> FDetailLayoutBuilderImpl::GetPropertyHandle( TSharedPtr<FPropertyNode> PropertyNodePtr )
+TSharedRef<IPropertyHandle> FDetailLayoutBuilderImpl::GetPropertyHandle( TSharedPtr<FPropertyNode> PropertyNodePtr ) const
 {
 	TSharedPtr<IPropertyHandle> PropertyHandle;
 	if( PropertyNodePtr.IsValid() )
diff --git a/Engine/Source/Editor/PropertyEditor/Private/DetailLayoutBuilderImpl.h b/Engine/Source/Editor/PropertyEditor/Private/DetailLayoutBuilderImpl.h
--- a/Engine/Source/Editor/PropertyEditor/Private/DetailLayoutBuilderImpl.h
+++ b/Engine/Source/Editor/PropertyEditor/Private/DetailLayoutBuilderImpl.h
@@ -36,7 +36,7 @@ public:
 	virtual IDetailPropertyRow& AddPropertyToCategory(TSharedPtr<IPropertyHandle> InPropertyHandle) override;
 	virtual FDetailWidgetRow& AddCustomRowToCategory(TSharedPtr<IPropertyHandle> InPropertyHandle, const FText& InCustomSearchString, bool bForAdvanced = false) override;
 	virtual IDetailPropertyRow* EditDefaultProperty(TSharedPtr<IPropertyHandle> InPropertyHandle) override;
-	virtual TSharedRef<IPropertyHandle> GetProperty(const FName PropertyPath, const UStruct* ClassOutermost, FName InInstanceName) override;
+	virtual TSharedRef<IPropertyHandle> GetProperty(const FName PropertyPath, const UStruct* ClassOutermost, FName InInstanceName) const override;
 	virtual FName GetTopLevelProperty() override;
 	virtual void HideProperty(const TSharedPtr<IPropertyHandle> Property) override;
 	virtual void HideProperty(FName PropertyPath, const UStruct* ClassOutermost = NULL, FName InstanceName = NAME_None) override;
@@ -170,7 +170,7 @@ public:
 	 *
 	 * @param PropertyNodePtr	The property node to make a handle for.
 	 */
-	TSharedRef<IPropertyHandle> GetPropertyHandle(TSharedPtr<FPropertyNode> PropertyNodePtr);
+	TSharedRef<IPropertyHandle> GetPropertyHandle(TSharedPtr<FPropertyNode> PropertyNodePtr) const;
 
 	/**
 	 * Adds an external property root node to the list of root nodes that the details new needs to manage
diff --git a/Engine/Source/Editor/PropertyEditor/Private/SDetailSingleItemRow.cpp b/Engine/Source/Editor/PropertyEditor/Private/SDetailSingleItemRow.cpp
--- a/Engine/Source/Editor/PropertyEditor/Private/SDetailSingleItemRow.cpp
+++ b/Engine/Source/Editor/PropertyEditor/Private/SDetailSingleItemRow.cpp
@@ -303,11 +303,12 @@ void SDetailSingleItemRow::Construct( const FArguments& InArgs, FDetailLayoutCus
 					Row.ValueWidget.Widget
 				];
 
-			ValueWidget = CreateExtensionWidget(ValueWidget.ToSharedRef(), *Customization, InOwnerTreeNode);
+			TSharedRef<SWidget> ExtensionWidget = CreateExtensionWidget(ValueWidget.ToSharedRef(), *Customization, InOwnerTreeNode);
 
 			if(Row.IsEnabledAttr.IsBound())
 			{
 				ValueWidget->SetEnabled(Row.IsEnabledAttr);
+				ExtensionWidget->SetEnabled(Row.IsEnabledAttr);
 			}
 
 			TSharedRef<SWidget> KeyFrameButton = CreateKeyframeButton(*Customization, InOwnerTreeNode);
@@ -431,21 +432,28 @@ void SDetailSingleItemRow::Construct( const FArguments& InArgs, FDetailLayoutCus
 					.Value(ColumnSizeData.RightColumnWidth)
 					.OnSlotResized(ColumnSizeData.OnWidthChanged)
 					[
-						SNew(SHorizontalBox)
-						+ SHorizontalBox::Slot()
+						SNew(SBox)
+						.IsEnabled(IsPropertyEditingEnabled)
 						[
-							SAssignNew(HBox, SHorizontalBox)
-							+ SHorizontalBox::Slot()
-							.Padding(DetailWidgetConstants::RightRowPadding)
-							.HAlign(Row.ValueWidget.HorizontalAlignment)
-							.VAlign(Row.ValueWidget.VerticalAlignment)
+							SNew(SHorizontalBox)
+							+SHorizontalBox::Slot()
+							.FillWidth(1.0f)
 							[
-								SNew(SBox)
-								.IsEnabled(IsPropertyEditingEnabled)
+								SNew(SHorizontalBox)
+								+ SHorizontalBox::Slot()
+								.Padding(DetailWidgetConstants::RightRowPadding)
+								.HAlign(Row.ValueWidget.HorizontalAlignment)
+								.VAlign(Row.ValueWidget.VerticalAlignment)
 								[
 									ValueWidget.ToSharedRef()
 								]
 							]
+							+SHorizontalBox::Slot()
+							.AutoWidth()
+							.VAlign(VAlign_Center)
+							[
+								ExtensionWidget
+							]
 						]
 					];
 				Widget = Splitter;
@@ -721,6 +729,8 @@ const FSlateBrush* SDetailSingleItemRow::GetBorderImage() const
 
 TSharedRef<SWidget> SDetailSingleItemRow::CreateExtensionWidget(TSharedRef<SWidget> ValueWidget, FDetailLayoutCustomization& InCustomization, TSharedRef<FDetailTreeNode> InTreeNode)
 {
+	TSharedPtr<SWidget> ExtensionWidget = SNullWidget::NullWidget;
+
 	if(InTreeNode->GetParentCategory().IsValid())
 	{
 		IDetailsViewPrivate* DetailsView = InTreeNode->GetDetailsView();
@@ -733,23 +743,13 @@ TSharedRef<SWidget> SDetailSingleItemRow::CreateExtensionWidget(TSharedRef<SWidg
 			UClass* ObjectClass = InCustomization.GetPropertyNode()->FindObjectItemParent()->GetObjectBaseClass();
 			if(Handle->IsValidHandle() && ExtensionHandler->IsPropertyExtendable(ObjectClass, *Handle))
 			{
-				ValueWidget = SNew(SHorizontalBox)
-					+ SHorizontalBox::Slot()
-					.FillWidth(1.0f)
-					[
-						ValueWidget
-					]
-
-					+ SHorizontalBox::Slot()
-					.AutoWidth()
-					[
-						ExtensionHandler->GenerateExtensionWidget(ObjectClass, Handle)
-					];
+				FDetailLayoutBuilderImpl& DetailLayout = OwnerTreeNode.Pin()->GetParentCategory()->GetParentLayoutImpl();
+				ExtensionWidget = ExtensionHandler->GenerateExtensionWidget(DetailLayout, ObjectClass, Handle);
 			}
 		}
 	}
 
-	return ValueWidget;
+	return ExtensionWidget.ToSharedRef();
 }
 
 TSharedRef<SWidget> SDetailSingleItemRow::CreateKeyframeButton( FDetailLayoutCustomization& InCustomization, TSharedRef<FDetailTreeNode> InTreeNode )
diff --git a/Engine/Source/Editor/PropertyEditor/Public/DetailLayoutBuilder.h b/Engine/Source/Editor/PropertyEditor/Public/DetailLayoutBuilder.h
--- a/Engine/Source/Editor/PropertyEditor/Public/DetailLayoutBuilder.h
+++ b/Engine/Source/Editor/PropertyEditor/Public/DetailLayoutBuilder.h
@@ -159,7 +159,7 @@ public:
 		To access StaticArray at index 2 from Struct2 in MyActor, your path would be MyStruct.StaticArray[2]" and your instance name is "Struct2"
 		To access MyFloat in MyActor you can just pass in "MyFloat" because the name of the property is unambiguous
 	 */
-	virtual TSharedRef<IPropertyHandle> GetProperty( const FName PropertyPath, const UStruct* ClassOutermost = NULL, FName InstanceName = NAME_None )  = 0;
+	virtual TSharedRef<IPropertyHandle> GetProperty( const FName PropertyPath, const UStruct* ClassOutermost = NULL, FName InstanceName = NAME_None ) const = 0;
 
 	/**
 	 * Gets the top level property, for showing the warning for experimental or early access class
diff --git a/Engine/Source/Editor/PropertyEditor/Public/IDetailPropertyExtensionHandler.h b/Engine/Source/Editor/PropertyEditor/Public/IDetailPropertyExtensionHandler.h
--- a/Engine/Source/Editor/PropertyEditor/Public/IDetailPropertyExtensionHandler.h
+++ b/Engine/Source/Editor/PropertyEditor/Public/IDetailPropertyExtensionHandler.h
@@ -6,6 +6,7 @@
 #include "Widgets/SWidget.h"
 
 class IPropertyHandle;
+class IDetailLayoutBuilder;
 
 class IDetailPropertyExtensionHandler
 {
@@ -14,5 +15,17 @@ public:
 
 	virtual bool IsPropertyExtendable(const UClass* InObjectClass, const class IPropertyHandle& PropertyHandle) const = 0;
 
-	virtual TSharedRef<SWidget> GenerateExtensionWidget(const UClass* InObjectClass, TSharedPtr<IPropertyHandle> PropertyHandle) = 0;
+	UE_DEPRECATED(4.24, "Please use the overload that takes a IDetailLayoutBuilder")
+	virtual TSharedRef<SWidget> GenerateExtensionWidget(const UClass* InObjectClass, TSharedPtr<IPropertyHandle> PropertyHandle) 
+	{ 
+		return SNullWidget::NullWidget;
+	}
+
+	virtual TSharedRef<SWidget> GenerateExtensionWidget(const IDetailLayoutBuilder& InDetailBuilder, const UClass* InObjectClass, TSharedPtr<IPropertyHandle> PropertyHandle)
+	{
+	PRAGMA_DISABLE_DEPRECATION_WARNINGS
+		// Call old deprecated path for back-compat
+		return GenerateExtensionWidget(InObjectClass, PropertyHandle);
+	PRAGMA_ENABLE_DEPRECATION_WARNINGS
+	}
 };
diff --git a/Engine/Source/Editor/SequenceRecorder/Private/SClassPropertyRecorderSettings.cpp b/Engine/Source/Editor/SequenceRecorder/Private/SClassPropertyRecorderSettings.cpp
--- a/Engine/Source/Editor/SequenceRecorder/Private/SClassPropertyRecorderSettings.cpp
+++ b/Engine/Source/Editor/SequenceRecorder/Private/SClassPropertyRecorderSettings.cpp
@@ -143,7 +143,7 @@ bool SClassPropertyRecorderSettings::IsPropertyExtendable(const UClass* InObject
 	return PropertyHandle.GetProperty()->HasAnyPropertyFlags(CPF_Interp) && FMovieSceneMultiPropertyRecorder::CanPropertyBeRecorded(*PropertyHandle.GetProperty());
 }
 
-TSharedRef<SWidget> SClassPropertyRecorderSettings::GenerateExtensionWidget(const UClass* InObjectClass, TSharedPtr<IPropertyHandle> PropertyHandle)
+TSharedRef<SWidget> SClassPropertyRecorderSettings::GenerateExtensionWidget(const IDetailLayoutBuilder& InDetailLayoutBuilder, const UClass* InObjectClass, TSharedPtr<IPropertyHandle> PropertyHandle)
 {
 	ECheckBoxState InitialState = ECheckBoxState::Unchecked;
 	FName PropertyName = *PropertyHandle->GeneratePathToProperty();
diff --git a/Engine/Source/Editor/SequenceRecorder/Private/SClassPropertyRecorderSettings.h b/Engine/Source/Editor/SequenceRecorder/Private/SClassPropertyRecorderSettings.h
--- a/Engine/Source/Editor/SequenceRecorder/Private/SClassPropertyRecorderSettings.h
+++ b/Engine/Source/Editor/SequenceRecorder/Private/SClassPropertyRecorderSettings.h
@@ -41,7 +41,7 @@ private:
 
 	/** IDetailPropertyExtensionHandler interface */
 	virtual bool IsPropertyExtendable(const UClass* InObjectClass, const class IPropertyHandle& PropertyHandle) const override;
-	virtual TSharedRef<SWidget> GenerateExtensionWidget(const UClass* InObjectClass, TSharedPtr<IPropertyHandle> PropertyHandle) override;
+	virtual TSharedRef<SWidget> GenerateExtensionWidget(const IDetailLayoutBuilder& InDetailLayoutBuilder, const UClass* InObjectClass, TSharedPtr<IPropertyHandle> PropertyHandle) override;
 
 private:
 	/** Handle to the property we are editing */
diff --git a/Engine/Source/Editor/Sequencer/Private/Sequencer.cpp b/Engine/Source/Editor/Sequencer/Private/Sequencer.cpp
--- a/Engine/Source/Editor/Sequencer/Private/Sequencer.cpp
+++ b/Engine/Source/Editor/Sequencer/Private/Sequencer.cpp
@@ -5670,6 +5670,121 @@ void FSequencer::GetSelectedFolders(TArray<UMovieSceneFolder*>& OutSelectedFolde
 	CalculateSelectedFolderAndPath(OutSelectedFolders, OutNewNodePath);
 }
 
+void FSequencer::GetSelectedObjects(TArray<FGuid>& Objects)
+{
+	Objects = GetSelection().GetBoundObjectsGuids();
+}
+
+void FSequencer::GetSelectedKeyAreas(TArray<const IKeyArea*>& OutSelectedKeyAreas)
+{
+	TSet<TSharedRef<FSequencerDisplayNode>> NodesToKey = Selection.GetSelectedOutlinerNodes();
+	{
+		TSet<TSharedRef<FSequencerDisplayNode>> ChildNodes;
+		for (TSharedRef<FSequencerDisplayNode> Node : NodesToKey.Array())
+		{
+			ChildNodes.Reset();
+			SequencerHelpers::GetDescendantNodes(Node, ChildNodes);
+
+			for (TSharedRef<FSequencerDisplayNode> ChildNode : ChildNodes)
+			{
+				NodesToKey.Remove(ChildNode);
+			}
+		}
+	}
+
+	TSet<TSharedPtr<IKeyArea>> KeyAreas;
+	TSet<UMovieSceneSection*>  ModifiedSections;
+
+	for (TSharedRef<FSequencerDisplayNode> Node : NodesToKey)
+	{
+		SequencerHelpers::GetAllKeyAreas(Node, KeyAreas);
+	}
+	for (FSequencerSelectedKey Key : Selection.GetSelectedKeys())
+	{
+		KeyAreas.Add(Key.KeyArea); 
+	}
+	for (TSharedPtr<IKeyArea> KeyArea : KeyAreas)
+	{
+		const IKeyArea* KeyAreaPtr = KeyArea.Get();
+		OutSelectedKeyAreas.Add(KeyAreaPtr);
+	}
+}
+
+void FSequencer::SelectByKeyAreas(const TArray<IKeyArea>& InKeyAreas, bool bSelectParentInstead, bool bSelect)
+{
+	TSet<TSharedRef<FSequencerDisplayNode>> Nodes;
+	TArray<TSharedRef<FSequencerDisplayNode>> NodesToSelect;
+	for (const TSharedRef<FSequencerDisplayNode>& Node : NodeTree->GetAllNodes())
+	{
+		if (Node->GetType() == ESequencerNode::Track)
+		{
+			TSharedRef<FSequencerTrackNode> TrackNode = StaticCastSharedRef<FSequencerTrackNode>(Node);
+			TArray<TSharedRef<FSequencerSectionKeyAreaNode>> KeyAreaNodes;
+			TrackNode->GetChildKeyAreaNodesRecursively(KeyAreaNodes);
+			for (TSharedRef<FSequencerSectionKeyAreaNode> KeyAreaNode : KeyAreaNodes)
+			{
+				for (TSharedPtr<IKeyArea> KeyArea : KeyAreaNode->GetAllKeyAreas())
+				{
+					for (const IKeyArea& InKeyArea : InKeyAreas)
+					{
+						if (InKeyArea.GetOwningSection() == KeyArea->GetOwningSection() &&
+							InKeyArea.GetChannel() == KeyArea->GetChannel())
+						{
+							if (bSelectParentInstead || bSelect == false)
+							{
+								Nodes.Add(KeyAreaNode->GetParent()->AsShared());
+							}
+							if (!bSelectParentInstead || bSelect == false)
+							{
+								Nodes.Add(KeyAreaNode);
+							}
+						}
+					}
+				}
+			}
+		}
+		else if (Node->GetType() == ESequencerNode::KeyArea)
+		{
+			TSharedRef<FSequencerSectionKeyAreaNode> KeyAreaNode = StaticCastSharedRef<FSequencerSectionKeyAreaNode>(Node);
+
+			for (TSharedPtr<IKeyArea> KeyArea : KeyAreaNode->GetAllKeyAreas())
+			{
+				for (const IKeyArea& InKeyArea : InKeyAreas)
+				{
+					if (InKeyArea.GetOwningSection() == KeyArea->GetOwningSection() &&
+						InKeyArea.GetChannel() == KeyArea->GetChannel())
+					{
+						if (bSelectParentInstead || bSelect == false)
+						{
+							Nodes.Add(KeyAreaNode->GetParent()->AsShared());
+						}
+						if (!bSelectParentInstead || bSelect == false)
+						{
+							Nodes.Add(KeyAreaNode);
+						}
+					}
+				}
+			}
+		}
+	}
+	if (bSelect)
+	{
+		for (const TSharedRef<FSequencerDisplayNode>& DisplayNode : Nodes)
+		{
+			NodesToSelect.Add(DisplayNode);
+			Selection.AddToSelection(NodesToSelect);
+		}
+	}
+	else
+	{
+		for (const TSharedRef<FSequencerDisplayNode>& DisplayNode : Nodes)
+		{
+			Selection.RemoveFromSelection(DisplayNode);
+			Selection.RemoveFromNodesWithSelectedKeysOrSections(DisplayNode);
+		}
+	}
+}
+
 void FSequencer::SelectObject(FGuid ObjectBinding)
 {
 	TSharedPtr<FSequencerObjectBindingNode> Node = NodeTree->FindObjectBindingNode(ObjectBinding);
diff --git a/Engine/Source/Editor/Sequencer/Private/Sequencer.h b/Engine/Source/Editor/Sequencer/Private/Sequencer.h
--- a/Engine/Source/Editor/Sequencer/Private/Sequencer.h
+++ b/Engine/Source/Editor/Sequencer/Private/Sequencer.h
@@ -728,10 +728,13 @@ public:
 	virtual void GetSelectedTracks(TArray<UMovieSceneTrack*>& OutSelectedTracks) override;
 	virtual void GetSelectedSections(TArray<UMovieSceneSection*>& OutSelectedSections) override;
 	virtual void GetSelectedFolders(TArray<UMovieSceneFolder*>& OutSelectedFolders) override;
+	virtual void GetSelectedKeyAreas(TArray<const IKeyArea*>& OutSelectedKeyAreas)  override;
+	virtual void GetSelectedObjects(TArray<FGuid>& OutObjects);
 	virtual void SelectObject(FGuid ObjectBinding) override;
 	virtual void SelectTrack(UMovieSceneTrack* Track) override;
 	virtual void SelectSection(UMovieSceneSection* Section) override;
 	virtual void SelectByPropertyPaths(const TArray<FString>& InPropertyPaths) override;
+	virtual void SelectByKeyAreas(const TArray<IKeyArea>& InKeyAreas, bool bSelectParentInstead, bool bSelect) override;
 	virtual void EmptySelection() override;
 	virtual void ThrobKeySelection() override;
 	virtual void ThrobSectionSelection() override;
diff --git a/Engine/Source/Editor/Sequencer/Private/SequencerContextMenus.cpp b/Engine/Source/Editor/Sequencer/Private/SequencerContextMenus.cpp
--- a/Engine/Source/Editor/Sequencer/Private/SequencerContextMenus.cpp
+++ b/Engine/Source/Editor/Sequencer/Private/SequencerContextMenus.cpp
@@ -726,7 +726,7 @@ void FSectionContextMenu::SetSectionToKey()
 	{
 		if (UMovieSceneSection* Section = WeakSection.Get())
 		{
-			UMovieScenePropertyTrack* Track = Section->GetTypedOuter<UMovieScenePropertyTrack>();
+			UMovieSceneTrack* Track = Section->GetTypedOuter<UMovieSceneTrack>();
 			if (Track)
 			{
 				FScopedTransaction Transaction(LOCTEXT("SetSectionToKey", "Set Section To Key"));
@@ -749,7 +749,7 @@ bool FSectionContextMenu::CanSetSectionToKey() const
 	{
 		if (UMovieSceneSection* Section = WeakSection.Get())
 		{
-			UMovieScenePropertyTrack* Track = Section->GetTypedOuter<UMovieScenePropertyTrack>();
+			UMovieSceneTrack* Track = Section->GetTypedOuter<UMovieSceneTrack>();
 			if (Track && Section->GetBlendType().IsValid() && (Section->GetBlendType().Get() == EMovieSceneBlendType::Absolute || Section->GetBlendType().Get() == EMovieSceneBlendType::Additive))
 			{
 				return true;
diff --git a/Engine/Source/Editor/Sequencer/Public/ISequencer.h b/Engine/Source/Editor/Sequencer/Public/ISequencer.h
--- a/Engine/Source/Editor/Sequencer/Public/ISequencer.h
+++ b/Engine/Source/Editor/Sequencer/Public/ISequencer.h
@@ -28,6 +28,7 @@ class FSequencerKeyCollection;
 class UMovieSceneSequence;
 class UMovieSceneSubSection;
 class IDetailsView;
+class IKeyArea;
 enum class EMapChangeType : uint8;
 
 /**
@@ -474,6 +475,12 @@ public:
 	/** Gets the currently selected folders. */
 	virtual void GetSelectedFolders(TArray<UMovieSceneFolder*>& OutSelectedFolders) = 0;
 
+	/** Gets the currently selected key areas */
+	virtual void GetSelectedKeyAreas(TArray<const IKeyArea*>& OutSelectedKeyAreas) = 0;
+
+	/** Gets the currently selected Object Guids*/
+	virtual void GetSelectedObjects(TArray<FGuid>& OutSelectedObjects) = 0;
+
 	/** Selects an object by GUID */
 	virtual void SelectObject(FGuid ObjectBinding) = 0;
 
@@ -486,6 +493,9 @@ public:
 	/** Selects property tracks by property path */
 	virtual void SelectByPropertyPaths(const TArray<FString>& InPropertyPaths) = 0;
 
+	/** Selects nodes by key areas */
+	virtual void SelectByKeyAreas(const TArray<IKeyArea>& InKeyAreas, bool bSelectParentInstead, bool bSelect) = 0;
+
 	/** Empties the current selection. */
 	virtual void EmptySelection() = 0;
 
diff --git a/Engine/Source/Editor/SkeletonEditor/Private/SSkeletonTree.cpp b/Engine/Source/Editor/SkeletonEditor/Private/SSkeletonTree.cpp
--- a/Engine/Source/Editor/SkeletonEditor/Private/SSkeletonTree.cpp
+++ b/Engine/Source/Editor/SkeletonEditor/Private/SSkeletonTree.cpp
@@ -906,12 +906,20 @@ bool GetSourceNameFromItem(TSharedPtr<ISkeletonTreeItem> SourceBone, FName& OutN
 void SSkeletonTree::FillVirtualBoneSubmenu(FMenuBuilder& MenuBuilder, TArray<TSharedPtr<ISkeletonTreeItem>> SourceBones)
 {
 	const bool bShowVirtualBones = false;
-	TSharedRef<SWidget> MenuContent = SNew(SBoneTreeMenu)
+	TSharedRef<SBoneTreeMenu> MenuContent = SNew(SBoneTreeMenu)
 		.bShowVirtualBones(false)
 		.Title(LOCTEXT("TargetBonePickerTitle", "Pick Target Bone..."))
 		.OnBoneSelectionChanged(this, &SSkeletonTree::OnVirtualTargetBonePicked, SourceBones)
 		.OnGetReferenceSkeleton(this, &SSkeletonTree::OnGetReferenceSkeleton);
 	MenuBuilder.AddWidget(MenuContent, FText::GetEmpty(), true);
+
+	MenuContent->RegisterActiveTimer(0.0f, FWidgetActiveTimerDelegate::CreateLambda(
+		[FilterTextBox = MenuContent->GetFilterTextWidget()](double, float)
+		{
+			FSlateApplication::Get().SetKeyboardFocus(FilterTextBox);
+			return EActiveTimerReturnType::Stop;
+		}
+	));
 }
 
 void SSkeletonTree::OnVirtualTargetBonePicked(FName TargetBoneName, TArray<TSharedPtr<ISkeletonTreeItem>> SourceBones)
diff --git a/Engine/Source/Editor/UMGEditor/Private/Details/DetailWidgetExtensionHandler.cpp b/Engine/Source/Editor/UMGEditor/Private/Details/DetailWidgetExtensionHandler.cpp
--- a/Engine/Source/Editor/UMGEditor/Private/Details/DetailWidgetExtensionHandler.cpp
+++ b/Engine/Source/Editor/UMGEditor/Private/Details/DetailWidgetExtensionHandler.cpp
@@ -42,7 +42,7 @@ bool FDetailWidgetExtensionHandler::IsPropertyExtendable(const UClass* InObjectC
 	return false;
 }
 
-TSharedRef<SWidget> FDetailWidgetExtensionHandler::GenerateExtensionWidget(const UClass* InObjectClass, TSharedPtr<IPropertyHandle> InPropertyHandle)
+TSharedRef<SWidget> FDetailWidgetExtensionHandler::GenerateExtensionWidget(const IDetailLayoutBuilder& InDetailBuilder, const UClass* InObjectClass, TSharedPtr<IPropertyHandle> InPropertyHandle)
 {
 	UProperty* Property = InPropertyHandle->GetProperty();
 	FString DelegateName = Property->GetName() + "Delegate";
diff --git a/Engine/Source/Editor/UMGEditor/Private/Details/DetailWidgetExtensionHandler.h b/Engine/Source/Editor/UMGEditor/Private/Details/DetailWidgetExtensionHandler.h
--- a/Engine/Source/Editor/UMGEditor/Private/Details/DetailWidgetExtensionHandler.h
+++ b/Engine/Source/Editor/UMGEditor/Private/Details/DetailWidgetExtensionHandler.h
@@ -16,7 +16,7 @@ public:
 
 	virtual bool IsPropertyExtendable(const UClass* InObjectClass, const class IPropertyHandle& PropertyHandle) const override;
 
-	virtual TSharedRef<SWidget> GenerateExtensionWidget(const UClass* InObjectClass, TSharedPtr<IPropertyHandle> PropertyHandle) override;
+	virtual TSharedRef<SWidget> GenerateExtensionWidget(const IDetailLayoutBuilder& InDetailBuilder, const UClass* InObjectClass, TSharedPtr<IPropertyHandle> PropertyHandle) override;
 
 private:
 	TWeakPtr<class FWidgetBlueprintEditor> BlueprintEditor;
diff --git a/Engine/Source/Editor/UnrealEd/Private/Fbx/FbxSkeletalMeshImport.cpp b/Engine/Source/Editor/UnrealEd/Private/Fbx/FbxSkeletalMeshImport.cpp
--- a/Engine/Source/Editor/UnrealEd/Private/Fbx/FbxSkeletalMeshImport.cpp
+++ b/Engine/Source/Editor/UnrealEd/Private/Fbx/FbxSkeletalMeshImport.cpp
@@ -1783,15 +1783,6 @@ USkeletalMesh* UnFbx::FFbxImporter::ImportSkeletalMesh(FImportSkeletalMeshArgs &
 
 	// Create initial bounding box based on expanded version of reference pose for meshes without physics assets. Can be overridden by artist.
 	FBox BoundingBox(SkelMeshImportDataPtr->Points.GetData(), SkelMeshImportDataPtr->Points.Num());
-	FBox Temp = BoundingBox;
-	FVector MidMesh = 0.5f*(Temp.Min + Temp.Max);
-	BoundingBox.Min = Temp.Min + 1.0f*(Temp.Min - MidMesh);
-	BoundingBox.Max = Temp.Max + 1.0f*(Temp.Max - MidMesh);
-	// Tuck up the bottom as this rarely extends lower than a reference pose's (e.g. having its feet on the floor).
-	// Maya has Y in the vertical, other packages have Z.
-	//BEN const int32 CoordToTuck = bAssumeMayaCoordinates ? 1 : 2;
-	//BEN BoundingBox.Min[CoordToTuck]	= Temp.Min[CoordToTuck] + 0.1f*(Temp.Min[CoordToTuck] - MidMesh[CoordToTuck]);
-	BoundingBox.Min[2] = Temp.Min[2] + 0.1f*(Temp.Min[2] - MidMesh[2]);
 	const FVector BoundingBoxSize = BoundingBox.GetSize();
 
 	if (SkelMeshImportDataPtr->Points.Num() > 2 && BoundingBoxSize.X < THRESH_POINTS_ARE_SAME && BoundingBoxSize.Y < THRESH_POINTS_ARE_SAME && BoundingBoxSize.Z < THRESH_POINTS_ARE_SAME)
diff --git a/Engine/Source/Editor/UnrealEd/Private/Kismet2/BlueprintEditorUtils.cpp b/Engine/Source/Editor/UnrealEd/Private/Kismet2/BlueprintEditorUtils.cpp
--- a/Engine/Source/Editor/UnrealEd/Private/Kismet2/BlueprintEditorUtils.cpp
+++ b/Engine/Source/Editor/UnrealEd/Private/Kismet2/BlueprintEditorUtils.cpp
@@ -6920,7 +6920,11 @@ void FBlueprintEditorUtils::UpdateOutOfDateCompositeWithOuter(UBlueprint* Bluepr
 		}
 		else if (UAnimStateTransitionNode* TransitionNode = Cast<UAnimStateTransitionNode>(Node))
 		{
-			BoundGraph = TransitionNode->BoundGraph;
+			// Shared rule graphs are not necessarily outered to their node, so skip them here
+			if(!TransitionNode->bSharedRules)
+			{
+				BoundGraph = TransitionNode->BoundGraph;
+			}
 		}
 		else if (UAnimGraphNode_StateMachineBase* StateMachineNode = Cast<UAnimGraphNode_StateMachineBase>(Node))
 		{
diff --git a/Engine/Source/Editor/UnrealEd/Private/SSkeletonWidget.cpp b/Engine/Source/Editor/UnrealEd/Private/SSkeletonWidget.cpp
--- a/Engine/Source/Editor/UnrealEd/Private/SSkeletonWidget.cpp
+++ b/Engine/Source/Editor/UnrealEd/Private/SSkeletonWidget.cpp
@@ -37,6 +37,7 @@ void SSkeletonListWidget::Construct(const FArguments& InArgs)
 	AssetPickerConfig.Filter.ClassNames.Add(USkeleton::StaticClass()->GetFName());
 	AssetPickerConfig.OnAssetSelected = FOnAssetSelected::CreateSP(this, &SSkeletonListWidget::SkeletonSelectionChanged);
 	AssetPickerConfig.InitialAssetViewType = EAssetViewType::Column;
+	AssetPickerConfig.SelectionMode = ESelectionMode::Single;
 	AssetPickerConfig.bShowPathInColumnView = true;
 	AssetPickerConfig.bShowTypeInColumnView = false;
 
diff --git a/Engine/Source/Programs/UnrealHeaderTool/Resources/UHTDebugging/DocumentationObjects.h b/Engine/Source/Programs/UnrealHeaderTool/Resources/UHTDebugging/DocumentationObjects.h
--- a/Engine/Source/Programs/UnrealHeaderTool/Resources/UHTDebugging/DocumentationObjects.h
+++ b/Engine/Source/Programs/UnrealHeaderTool/Resources/UHTDebugging/DocumentationObjects.h
@@ -13,4 +13,17 @@ UCLASS(meta=(DocumentationPolicy = "Strict"))
 class UClassToDocument : public UObject
 {
 	GENERATED_BODY()
+
+	UClassToDocument()
+		: MyProperty(0.f)
+	{
+	}
+
+	// A property for testing the policy
+	UPROPERTY(meta = (UIMin = "0.0", UIMax = "1.0"))
+	float MyProperty;
+
+	// A function to test the policy
+	UFUNCTION()
+	void MyFunction() {}
 };
diff --git a/Engine/Source/Runtime/AnimGraphRuntime/Private/AnimNodes/AnimNode_BlendListBase.cpp b/Engine/Source/Runtime/AnimGraphRuntime/Private/AnimNodes/AnimNode_BlendListBase.cpp
--- a/Engine/Source/Runtime/AnimGraphRuntime/Private/AnimNodes/AnimNode_BlendListBase.cpp
+++ b/Engine/Source/Runtime/AnimGraphRuntime/Private/AnimNodes/AnimNode_BlendListBase.cpp
@@ -113,10 +113,6 @@ void FAnimNode_BlendListBase::Update_AnyThread(const FAnimationUpdateContext& Co
 				{
 					InertializationNode->Request(BlendTime[ChildIndex]);
 				}
-				else
-				{
-					UE_LOG(LogAnimation, Error, TEXT("FAnimNode_BlendListBase: No inertialization context found for transition"));
-				}
 				
 				RemainingBlendTime = 0.0f;
 			}
diff --git a/Engine/Source/Runtime/AnimGraphRuntime/Private/AnimNodes/AnimNode_MultiWayBlend.cpp b/Engine/Source/Runtime/AnimGraphRuntime/Private/AnimNodes/AnimNode_MultiWayBlend.cpp
--- a/Engine/Source/Runtime/AnimGraphRuntime/Private/AnimNodes/AnimNode_MultiWayBlend.cpp
+++ b/Engine/Source/Runtime/AnimGraphRuntime/Private/AnimNodes/AnimNode_MultiWayBlend.cpp
@@ -83,6 +83,8 @@ void FAnimNode_MultiWayBlend::UpdateCachedAlphas()
 			}
 		}
 	}
+
+	ensure(Poses.Num() == CachedAlphas.Num());
 }
 
 void FAnimNode_MultiWayBlend::Update_AnyThread(const FAnimationUpdateContext& Context)
diff --git a/Engine/Source/Runtime/AnimGraphRuntime/Private/AnimSequencerInstanceProxy.cpp b/Engine/Source/Runtime/AnimGraphRuntime/Private/AnimSequencerInstanceProxy.cpp
--- a/Engine/Source/Runtime/AnimGraphRuntime/Private/AnimSequencerInstanceProxy.cpp
+++ b/Engine/Source/Runtime/AnimGraphRuntime/Private/AnimSequencerInstanceProxy.cpp
@@ -92,7 +92,6 @@ void FAnimSequencerInstanceProxy::InitAnimTrack(UAnimSequenceBase* InAnimSequenc
 			}
 
 			const int32 PoseIndex = BlendNode.AddPose() - 1;
-			BlendNode.UpdateCachedAlphas();
 
 			// add the new entry to map
 			FSequencerPlayerAnimSequence* NewPlayerState = new FSequencerPlayerAnimSequence();
