Fix nativized CIS automation test regressions due to incorrect default values being emitted to the generated ctor code for converted BP class assets.

#rnx

diff --git a/Engine/Source/Editor/Kismet/Private/BlueprintCompilationManager.cpp b/Engine/Source/Editor/Kismet/Private/BlueprintCompilationManager.cpp
--- a/Engine/Source/Editor/Kismet/Private/BlueprintCompilationManager.cpp
+++ b/Engine/Source/Editor/Kismet/Private/BlueprintCompilationManager.cpp
@@ -422,6 +422,7 @@ struct FCompilerData
 	}
 
 	bool IsSkeletonOnly() const { return JobType == ECompilationManagerJobType::SkeletonOnly; }
+	bool IsCppCompileType() const { return InternalOptions.CompileType == EKismetCompileType::Cpp; }
 	bool ShouldSetTemporaryBlueprintFlags() const { return JobType != ECompilationManagerJobType::RelinkOnly; }
 	bool ShouldResetErrorState() const { return JobType == ECompilationManagerJobType::Normal && InternalOptions.CompileType != EKismetCompileType::BytecodeOnly; }
 	bool ShouldValidateVariableNames() const { return JobType == ECompilationManagerJobType::Normal; }
@@ -1083,10 +1084,17 @@ void FBlueprintCompilationManagerImpl::FlushCompilationQueueImpl(bool bSuppressB
 					OldCDOs.Add(BP, BP->GeneratedClass->ClassDefaultObject);
 				}
 
+				EBlueprintCompileReinstancerFlags CompileReinstancerFlags =
+					EBlueprintCompileReinstancerFlags::AutoInferSaveOnCompile;
+				if (!CompilerData.IsCppCompileType())
+				{
+					CompileReinstancerFlags |= EBlueprintCompileReinstancerFlags::AvoidCDODuplication;
+				}
+
 				CompilerData.Reinstancer = TSharedPtr<FBlueprintCompileReinstancer>(
 					new FBlueprintCompileReinstancer(
 						BP->GeneratedClass,
-						EBlueprintCompileReinstancerFlags::AutoInferSaveOnCompile | EBlueprintCompileReinstancerFlags::AvoidCDODuplication
+						CompileReinstancerFlags
 					)
 				);
 
@@ -1127,7 +1135,7 @@ void FBlueprintCompilationManagerImpl::FlushCompilationQueueImpl(bool bSuppressB
 				// default value propagation occurs in ReinstaneBatch, CDO will be created via CompileFunctions call:
 				if(BP->ParentClass)
 				{
-					if(BP->GeneratedClass)
+					if(BP->GeneratedClass && !CompilerData.IsCppCompileType())
 					{
 						BP->GeneratedClass->ClassDefaultObject = nullptr;
 					}
diff --git a/Engine/Source/Editor/KismetCompiler/Private/KismetCompiler.cpp b/Engine/Source/Editor/KismetCompiler/Private/KismetCompiler.cpp
--- a/Engine/Source/Editor/KismetCompiler/Private/KismetCompiler.cpp
+++ b/Engine/Source/Editor/KismetCompiler/Private/KismetCompiler.cpp
@@ -4094,13 +4094,26 @@ void FKismetCompilerContext::CompileFunctions(EInternalCompilerFlags InternalFla
 	// This is phase two, so we want to generated locals if PostponeLocalsGenerationUntilPhaseTwo is set:
 	const bool bGenerateLocals = !!(InternalFlags & EInternalCompilerFlags::PostponeLocalsGenerationUntilPhaseTwo);
 	// Don't propagate values to CDO if we're going to do that in reinstancing:
-	const bool bPropagateValuesToCDO = !(InternalFlags & EInternalCompilerFlags::PostponeDefaultObjectAssignmentUntilReinstancing);
+	bool bPropagateValuesToCDO = !(InternalFlags & EInternalCompilerFlags::PostponeDefaultObjectAssignmentUntilReinstancing);
 	// Don't RefreshExternalBlueprintDependencyNodes if the calling code has done so already:
 	const bool bSkipRefreshExternalBlueprintDependencyNodes = !!(InternalFlags & EInternalCompilerFlags::SkipRefreshExternalBlueprintDependencyNodes);
 	FKismetCompilerVMBackend Backend_VM(Blueprint, Schema, *this);
 
-	// Determine whether or not to skip generated class validation. This requires CDO value propagation to occur first.
-	bool bSkipGeneratedClassValidation = !bPropagateValuesToCDO || CompileOptions.CompileType == EKismetCompileType::Cpp;
+	// Determine whether or not to skip generated class validation.
+	bool bSkipGeneratedClassValidation;
+	if (CompileOptions.DoesRequireCppCodeGeneration())
+	{
+		// CPP codegen requires default value assignment to occur as part of the compilation phase, so we override it here.
+		bPropagateValuesToCDO = true;
+
+		// Also skip generated class validation since it may result in errors and we don't really need to keep the generated class.
+		bSkipGeneratedClassValidation = true;
+	}
+	else
+	{
+		// In all other cases, validation requires CDO value propagation to occur first.
+		bSkipGeneratedClassValidation = !bPropagateValuesToCDO;
+	}
 
 	if( bGenerateLocals )
 	{
