Merging using //UE4/Dev-Framework_to_//UE4/Release-4.24

Re-order how we do sparse properties so that it will use sparse properties, if they exist, in blueprint nodes.

diff --git a/Engine/Source/Editor/KismetCompiler/Private/KismetCompilerMisc.cpp b/Engine/Source/Editor/KismetCompiler/Private/KismetCompilerMisc.cpp
--- a/Engine/Source/Editor/KismetCompiler/Private/KismetCompilerMisc.cpp
+++ b/Engine/Source/Editor/KismetCompiler/Private/KismetCompilerMisc.cpp
@@ -518,12 +518,28 @@ void FKismetCompilerUtilities::RemoveObjectRedirectorIfPresent(UObject* Package,
 }
 
 /** Finds a property by name, starting in the specified scope; Validates property type and returns NULL along with emitting an error if there is a mismatch. */
-UProperty* FKismetCompilerUtilities::FindPropertyInScope(UStruct* Scope, UEdGraphPin* Pin, FCompilerResultsLog& MessageLog, const UEdGraphSchema_K2* Schema, UClass* SelfClass, bool& bIsSparseProperty)
+UProperty* FKismetCompilerUtilities::FindPropertyInScope(UStruct* Scope, UEdGraphPin* Pin, FCompilerResultsLog& MessageLog, const UEdGraphSchema_K2* Schema, UClass* SelfClass, bool& bIsSparseProperty, bool bSuppressMissingMemberErrors)
 {
 	bIsSparseProperty = false;
 	UStruct* InitialScope = Scope;
-	while (Scope != NULL)
+	while (Scope != nullptr)
 	{
+		// If this is a class, check the sparse data for the property
+		UClass* Class = Cast<UClass>(Scope);
+		if (Class)
+		{
+			UStruct* SparseData = Class->GetSparseClassDataStruct();
+			if (SparseData)
+			{
+				UProperty* Prop = FindPropertyInScope(SparseData, Pin, MessageLog, Schema, SelfClass, bIsSparseProperty, true);
+				if (Prop)
+				{
+					bIsSparseProperty = true;
+					return Prop;
+				}
+			}
+		}
+
 		for (TFieldIterator<UProperty> It(Scope, EFieldIteratorFlags::IncludeSuper); It; ++It)
 		{
 			UProperty* Property = *It;
@@ -537,38 +553,22 @@ UProperty* FKismetCompilerUtilities::FindPropertyInScope(UStruct* Scope, UEdGrap
 				else
 				{
 					// Exit now, we found one with the right name but the type mismatched (and there was a type mismatch error)
-					return NULL;
-				}
-			}
-		}
-
-		// If this is a class, check the sparse data for the property
-		UClass* Class = Cast<UClass>(Scope);
-		if (Class)
-		{
-			UStruct* SparseData = Class->GetSparseClassDataStruct();
-			if (SparseData)
-			{
-				UProperty* Prop = FindPropertyInScope(SparseData, Pin, MessageLog, Schema, SelfClass, bIsSparseProperty);
-				if (Prop)
-				{
-					bIsSparseProperty = true;
-					return Prop;
+					return nullptr;
 				}
 			}
 		}
 
 		// Functions don't automatically check their class when using a field iterator
 		UFunction* Function = Cast<UFunction>(Scope);
-		Scope = (Function != NULL) ? Cast<UStruct>(Function->GetOuter()) : NULL;
+		Scope = (Function != nullptr) ? Cast<UStruct>(Function->GetOuter()) : nullptr;
 	}
 
 	// Couldn't find the name
-	if (!FKismetCompilerUtilities::IsMissingMemberPotentiallyLoading(Cast<UBlueprint>(SelfClass->ClassGeneratedBy), InitialScope))
+	if (!FKismetCompilerUtilities::IsMissingMemberPotentiallyLoading(Cast<UBlueprint>(SelfClass->ClassGeneratedBy), InitialScope) && !bSuppressMissingMemberErrors)
 	{
 		MessageLog.Error(*FText::Format(LOCTEXT("PropertyNotFound_Error", "The property associated with @@ could not be found in '{0}'"), FText::FromString(SelfClass->GetPathName())).ToString(), Pin);
 	}
-	return NULL;
+	return nullptr;
 }
 
 // Finds a property by name, starting in the specified scope, returning NULL if it's not found
diff --git a/Engine/Source/Editor/KismetCompiler/Public/KismetCompilerMisc.h b/Engine/Source/Editor/KismetCompiler/Public/KismetCompilerMisc.h
--- a/Engine/Source/Editor/KismetCompiler/Public/KismetCompilerMisc.h
+++ b/Engine/Source/Editor/KismetCompiler/Public/KismetCompilerMisc.h
@@ -51,7 +51,7 @@ public:
 	static bool IsTypeCompatibleWithProperty(UEdGraphPin* SourcePin, UProperty* Property, FCompilerResultsLog& MessageLog, const UEdGraphSchema_K2* Schema, UClass* SelfClass);
 
 	/** Finds a property by name, starting in the specified scope; Validates property type and returns NULL along with emitting an error if there is a mismatch. */
-	static UProperty* FindPropertyInScope(UStruct* Scope, UEdGraphPin* Pin, FCompilerResultsLog& MessageLog, const UEdGraphSchema_K2* Schema, UClass* SelfClass, bool& bIsSparseProperty);
+	static UProperty* FindPropertyInScope(UStruct* Scope, UEdGraphPin* Pin, FCompilerResultsLog& MessageLog, const UEdGraphSchema_K2* Schema, UClass* SelfClass, bool& bIsSparseProperty, bool bSuppressMissingMemberErrors = false);
 
 	// Finds a property by name, starting in the specified scope, returning NULL if it's not found
 	static UProperty* FindNamedPropertyInScope(UStruct* Scope, FName PropertyName, bool& bIsSparseProperty);
