Added missing error checking animation state machine states

This error checking has seemingly been missing for a long time

diff --git a/Engine/Source/Editor/KismetCompiler/Private/AnimBlueprintCompiler.cpp b/Engine/Source/Editor/KismetCompiler/Private/AnimBlueprintCompiler.cpp
--- a/Engine/Source/Editor/KismetCompiler/Private/AnimBlueprintCompiler.cpp
+++ b/Engine/Source/Editor/KismetCompiler/Private/AnimBlueprintCompiler.cpp
@@ -1556,6 +1556,8 @@ void FAnimBlueprintCompilerContext::ProcessStateMachine(UAnimGraphNode_StateMach
 				// Process the inner graph of this state
 				if (UAnimGraphNode_StateResult* AnimGraphResultNode = CastChecked<UAnimationStateGraph>(StateNode->BoundGraph)->GetResultNode())
 				{
+					ValidateGraphIsWellFormed(StateNode->BoundGraph);
+
 					BakedState.StateRootNodeIndex = ExpandGraphAndProcessNodes(StateNode->BoundGraph, AnimGraphResultNode);
 
 					// See if the state consists of a single sequence player node, and remember the index if so
