Prevent ICH from getting moved to transient package in case the transaction buffer brings it back, add a transaction to blueprint reparent operations so that they can be undone/redone coherently

diff --git a/Engine/Source/Editor/Kismet/Private/BlueprintEditor.cpp b/Engine/Source/Editor/Kismet/Private/BlueprintEditor.cpp
--- a/Engine/Source/Editor/Kismet/Private/BlueprintEditor.cpp
+++ b/Engine/Source/Editor/Kismet/Private/BlueprintEditor.cpp
@@ -2838,7 +2838,20 @@ void FBlueprintEditor::ReparentBlueprint_NewParentChosen(UClass* ChosenClass)
 
 		if ( bReparent )
 		{
+			const FScopedTransaction Transaction( LOCTEXT("ReparentBlueprint", "Reparent Blueprint") );
 			UE_LOG(LogBlueprint, Warning, TEXT("Reparenting blueprint %s from %s to %s..."), *BlueprintObj->GetFullName(), BlueprintObj->ParentClass ? *BlueprintObj->ParentClass->GetName() : TEXT("[None]"), *ChosenClass->GetName());
+			
+			BlueprintObj->Modify();
+			if(USimpleConstructionScript* SCS = BlueprintObj->SimpleConstructionScript)
+			{
+				SCS->Modify();
+
+				const TArray<USCS_Node*>& AllNodes = SCS->GetAllNodes();
+				for(USCS_Node* Node : AllNodes )
+				{
+					Node->Modify();
+				}
+			}
 
 			UClass* OldParentClass = BlueprintObj->ParentClass ;
 			BlueprintObj->ParentClass = ChosenClass;
diff --git a/Engine/Source/Editor/KismetCompiler/Private/KismetCompiler.cpp b/Engine/Source/Editor/KismetCompiler/Private/KismetCompiler.cpp
--- a/Engine/Source/Editor/KismetCompiler/Private/KismetCompiler.cpp
+++ b/Engine/Source/Editor/KismetCompiler/Private/KismetCompiler.cpp
@@ -272,9 +272,19 @@ void FKismetCompilerContext::CleanAndSanitizeClass(UBlueprintGeneratedClass* Cla
 		ClassSubObjects.RemoveAllSwap(SubObjectsToSave);
 	}
 
-	for( auto SubObjIt = ClassSubObjects.CreateIterator(); SubObjIt; ++SubObjIt )
+	UClass* InheritableComponentHandlerClass = UInheritableComponentHandler::StaticClass();
+
+	for( UObject* CurrSubObj : ClassSubObjects )
 	{
-		UObject* CurrSubObj = *SubObjIt;
+		// ICH and ICH templates do not need to be destroyed in this way.. doing so will invalidate
+		// transaction buffer references to these UObjects. The UBlueprint may not have a reference to
+		// the ICH at the moment, and therefore might not have added it to SubObjectsToSave (and
+		// removed the ICH from ClassSubObjects):
+		if(Cast<UInheritableComponentHandler>(CurrSubObj) || CurrSubObj->IsInA(InheritableComponentHandlerClass))
+		{
+			continue;
+		}
+
 		FName NewSubobjectName = MakeUniqueObjectName(TransientClass, CurrSubObj->GetClass(), CurrSubObj->GetFName());
 		CurrSubObj->Rename(*NewSubobjectName.ToString(), TransientClass, RenFlags);
 		if( UProperty* Prop = Cast<UProperty>(CurrSubObj) )
