ControlRig: Avoid custom pins on custom nodes with the same name

diff --git a/Engine/Source/Editor/KismetCompiler/Private/AnimBlueprintCompiler.cpp b/Engine/Source/Editor/KismetCompiler/Private/AnimBlueprintCompiler.cpp
--- a/Engine/Source/Editor/KismetCompiler/Private/AnimBlueprintCompiler.cpp
+++ b/Engine/Source/Editor/KismetCompiler/Private/AnimBlueprintCompiler.cpp
@@ -780,6 +780,16 @@ void FAnimBlueprintCompilerContext::ProcessCustomPropertyNode(UAnimGraphNode_Cus
 	{
 		if (!Pin->bOrphanedPin && !AnimGraphSchema->IsPosePin(Pin->PinType))
 		{
+			// avoid to add properties which already exist on the custom node.
+			// for example the ControlRig_CustomNode has a pin called "alpha" which is not custom.
+			if (UStructProperty* NodeProperty = Cast<UStructProperty>(CustomPropNode->GetClass()->FindPropertyByName(TEXT("Node"))))
+			{
+				if(NodeProperty->Struct->FindPropertyByName(Pin->GetFName()))
+				{
+					continue;
+				}
+			}
+
 			// Add prefix to avoid collisions
 			FString PrefixedName = CustomPropNode->GetPinTargetVariableName(Pin);
 
