Fixed non-determinstic cooking of conformed-on-load interface graphs

Switched anim nodes to name their node properties without using the node GUID, similar to to other graph nodes.
Also fixed old COL-non-determinism/duplicate GUID guard to include interface graphs.

diff --git a/Engine/Source/Editor/KismetCompiler/Private/AnimBlueprintCompiler.cpp b/Engine/Source/Editor/KismetCompiler/Private/AnimBlueprintCompiler.cpp
--- a/Engine/Source/Editor/KismetCompiler/Private/AnimBlueprintCompiler.cpp
+++ b/Engine/Source/Editor/KismetCompiler/Private/AnimBlueprintCompiler.cpp
@@ -137,12 +137,12 @@ FAnimBlueprintCompilerContext::FAnimBlueprintCompilerContext(UAnimBlueprint* Sou
 		// Tracking to see if we need to warn for deterministic cooking
 		bool bNodeGuidsRegenerated = false;
 
-		for (UEdGraph* Graph : AnimBlueprint->FunctionGraphs)
+		auto CheckGraph = [&bNodeGuidsRegenerated, &NodeGuids, &ChildGraphs](UEdGraph* InGraph)
 		{
-			if (AnimationEditorUtils::IsAnimGraph(Graph))
+			if (AnimationEditorUtils::IsAnimGraph(InGraph))
 			{
 				ChildGraphs.Reset();
-				AnimationEditorUtils::FindChildGraphsFromNodes(Graph, ChildGraphs);
+				AnimationEditorUtils::FindChildGraphsFromNodes(InGraph, ChildGraphs);
 
 				for (int32 Index = 0; Index < ChildGraphs.Num(); ++Index) // Not ranged for as we modify array within the loop
 				{
@@ -169,6 +169,19 @@ FAnimBlueprintCompilerContext::FAnimBlueprintCompilerContext(UAnimBlueprint* Sou
 					}
 				}
 			}
+		};
+
+		for (UEdGraph* Graph : AnimBlueprint->FunctionGraphs)
+		{
+			CheckGraph(Graph);
+		}
+
+		for(FBPInterfaceDescription& InterfaceDesc : Blueprint->ImplementedInterfaces)
+		{
+			for(UEdGraph* Graph : InterfaceDesc.Graphs)
+			{
+				CheckGraph(Graph);
+			}
 		}
 
 		if(bNodeGuidsRegenerated)
diff --git a/Engine/Source/Editor/KismetCompiler/Private/KismetCompilerMisc.cpp b/Engine/Source/Editor/KismetCompiler/Private/KismetCompilerMisc.cpp
--- a/Engine/Source/Editor/KismetCompiler/Private/KismetCompilerMisc.cpp
+++ b/Engine/Source/Editor/KismetCompiler/Private/KismetCompilerMisc.cpp
@@ -34,7 +34,6 @@
 #include "KismetCompiledFunctionContext.h"
 #include "KismetCompiler.h"
 
-#include "AnimGraphNode_Base.h"
 #include "K2Node_EnumLiteral.h"
 #include "Kismet2/KismetReinstanceUtilities.h"
 #include "Kismet2/BlueprintEditorUtils.h"
@@ -1817,11 +1816,6 @@ FString FNetNameMapping::MakeBaseName(const UEdGraphNode* Net)
 	return FString::Printf(TEXT("%s"), *Net->GetDescriptiveCompiledName());
 }
 
-FString FNetNameMapping::MakeBaseName(const UAnimGraphNode_Base* Net)
-{
-	return FString::Printf(TEXT("%s_%s"), *Net->GetDescriptiveCompiledName(), *Net->NodeGuid.ToString());
-}
-
 //////////////////////////////////////////////////////////////////////////
 // FKismetFunctionContext
 
diff --git a/Engine/Source/Editor/KismetCompiler/Public/KismetCompilerMisc.h b/Engine/Source/Editor/KismetCompiler/Public/KismetCompilerMisc.h
--- a/Engine/Source/Editor/KismetCompiler/Public/KismetCompilerMisc.h
+++ b/Engine/Source/Editor/KismetCompiler/Public/KismetCompilerMisc.h
@@ -8,7 +8,6 @@
 
 class FCompilerResultsLog;
 class FKismetCompilerContext;
-class UAnimGraphNode_Base;
 class UBlueprint;
 class UEdGraph;
 class UEdGraphSchema_K2;
@@ -210,12 +209,10 @@ public:
 	// The resulting name is stable across multiple calls if given the same pointer.
 	FString MakeValidName(const UEdGraphNode* Net, const FString& Context = TEXT("")) { return MakeValidNameImpl(Net, Context); }
 	FString MakeValidName(const UEdGraphPin* Net, const FString& Context = TEXT("")) { return MakeValidNameImpl(Net, Context); }
-	FString MakeValidName(const UAnimGraphNode_Base* Net, const FString& Context = TEXT("")) { return MakeValidNameImpl(Net, Context); }
 
 private:
 	KISMETCOMPILER_API static FString MakeBaseName(const UEdGraphNode* Net);
 	KISMETCOMPILER_API static FString MakeBaseName(const UEdGraphPin* Net);
-	KISMETCOMPILER_API static FString MakeBaseName(const UAnimGraphNode_Base* Net);
 
 	template< typename NetType >
 	FString MakeValidNameImpl(NetType Net, const FString& Context)
