Blueprint: Added support for generated functions within Blueprint Graphs

Generated function graphs can be created as part of the compilation process and are ephemeral (only lasting the duration of the compile), but will still be used to generate bytecode and UFunctions on the final class layout. External clients can hook into this process by binding to UBlueprint::GenerateFunctionGraphsEvent and spawning the necessary intermediate graphs on the compiler context using FKismetCompiler::SpawnIntermediateFunctionGraph. GenerateFunctionGraphsEvent is a dynamic delegate, meaning that all bindings will be serialized into the package, enabling composition workflows within a single package.

Generated function graphs are also visible as an intermediate build product with this Blueprint developer setting enabled.

diff --git a/Engine/Source/Editor/KismetCompiler/Private/KismetCompiler.cpp b/Engine/Source/Editor/KismetCompiler/Private/KismetCompiler.cpp
--- a/Engine/Source/Editor/KismetCompiler/Private/KismetCompiler.cpp
+++ b/Engine/Source/Editor/KismetCompiler/Private/KismetCompiler.cpp
@@ -73,6 +73,7 @@ DECLARE_CYCLE_STAT(TEXT("Create Function List"), EKismetCompilerStats_CreateFunc
 DECLARE_CYCLE_STAT(TEXT("Expansion"), EKismetCompilerStats_Expansion, STATGROUP_KismetCompiler )
 DECLARE_CYCLE_STAT(TEXT("Process uber"), EKismetCompilerStats_ProcessUbergraph, STATGROUP_KismetCompiler );
 DECLARE_CYCLE_STAT(TEXT("Process func"), EKismetCompilerStats_ProcessFunctionGraph, STATGROUP_KismetCompiler );
+DECLARE_CYCLE_STAT(TEXT("Generate Function Graph"), EKismetCompilerStats_GenerateFunctionGraphs, STATGROUP_KismetCompiler );
 DECLARE_CYCLE_STAT(TEXT("Precompile Function"), EKismetCompilerStats_PrecompileFunction, STATGROUP_KismetCompiler );
 DECLARE_CYCLE_STAT(TEXT("Compile Function"), EKismetCompilerStats_CompileFunction, STATGROUP_KismetCompiler );
 DECLARE_CYCLE_STAT(TEXT("Postcompile Function"), EKismetCompilerStats_PostcompileFunction, STATGROUP_KismetCompiler );
@@ -3780,9 +3781,11 @@ void FKismetCompilerContext::ValidateFunctionGraphNames()
 
 	if(ParentBPNameValidator.IsValid())
 	{
-		for (int32 FunctionIndex=0; FunctionIndex < Blueprint->FunctionGraphs.Num(); ++FunctionIndex)
+		TArray<UEdGraph*> AllFunctionGraphs(Blueprint->FunctionGraphs);
+		AllFunctionGraphs += GeneratedFunctionGraphs;
+
+		for (UEdGraph* FunctionGraph : AllFunctionGraphs)
 		{
-			UEdGraph* FunctionGraph = Blueprint->FunctionGraphs[FunctionIndex];
 			if(FunctionGraph->GetFName() != UEdGraphSchema_K2::FN_UserConstructionScript)
 			{
 				if( ParentBPNameValidator->IsValid(FunctionGraph->GetName()) != EValidatorResult::Ok )
@@ -3807,6 +3810,15 @@ void FKismetCompilerContext::ValidateFunctionGraphNames()
 // Creates a copy of the graph to allow further transformations to occur
 void FKismetCompilerContext::CreateFunctionList()
 {
+	{
+		BP_SCOPED_COMPILER_EVENT_STAT(EKismetCompilerStats_GenerateFunctionGraphs);
+
+		// Broadcast the blueprint's function generation event to ensure that any function generators are run right before we create the function list for the class layout
+		FGenerateBlueprintFunctionParams Params;
+		Params.CompilerContext = this;
+		Blueprint->GenerateFunctionGraphsEvent.Broadcast(Params);
+	}
+
 	BP_SCOPED_COMPILER_EVENT_STAT(EKismetCompilerStats_CreateFunctionList);
 
 	// Process the ubergraph if one should be present
@@ -3826,6 +3838,11 @@ void FKismetCompilerContext::CreateFunctionList()
 			ProcessOneFunctionGraph(Blueprint->FunctionGraphs[i]);
 		}
 
+		for (UEdGraph* FunctionGraph : GeneratedFunctionGraphs)
+		{
+			ProcessOneFunctionGraph(FunctionGraph);
+		}
+
 		for (int32 i = 0; i < Blueprint->DelegateSignatureGraphs.Num(); ++i)
 		{
 			// change function names to unique
@@ -4180,6 +4197,8 @@ void FKismetCompilerContext::CompileFunctions(EInternalCompilerFlags InternalFla
 		}
 	}
 
+	FunctionListCompiledEvent.Broadcast(this);
+
 	// Save off intermediate build products if requested
 	if (bIsFullCompile && CompileOptions.bSaveIntermediateProducts && !Blueprint->bIsRegeneratingOnLoad)
 	{
@@ -4618,6 +4637,21 @@ bool FKismetCompilerContext::ValidateGeneratedClass(UBlueprintGeneratedClass* Cl
 	return UBlueprint::ValidateGeneratedClass(Class);
 }
 
+UEdGraph* FKismetCompilerContext::SpawnIntermediateFunctionGraph(const FString& InDesiredFunctionName)
+{
+	FName UniqueGraphName = FBlueprintEditorUtils::FindUniqueKismetName(Blueprint, InDesiredFunctionName);
+
+	UEdGraph* GeneratedFunctionGraph = FBlueprintEditorUtils::CreateNewGraph(Blueprint, UniqueGraphName, UEdGraph::StaticClass(), UEdGraphSchema_K2::StaticClass());
+	GeneratedFunctionGraph->SetFlags(RF_Transient);
+	GeneratedFunctionGraph->bEditable = false;
+
+	FBlueprintEditorUtils::CreateFunctionGraph(Blueprint, GeneratedFunctionGraph, false, (UClass*)nullptr);
+
+	// Add the function graph to the list of generated graphs for this compile
+	GeneratedFunctionGraphs.Add(GeneratedFunctionGraph);
+	return GeneratedFunctionGraph;
+}
+
 const UK2Node_FunctionEntry* FKismetCompilerContext::FindLocalEntryPoint(const UFunction* Function) const
 {
 	for (int32 i = 0; i < FunctionList.Num(); ++i)
diff --git a/Engine/Source/Editor/KismetCompiler/Public/KismetCompiler.h b/Engine/Source/Editor/KismetCompiler/Public/KismetCompiler.h
--- a/Engine/Source/Editor/KismetCompiler/Public/KismetCompiler.h
+++ b/Engine/Source/Editor/KismetCompiler/Public/KismetCompiler.h
@@ -41,6 +41,10 @@ typedef TFunction<TSharedPtr<FKismetCompilerContext>(UBlueprint*, FCompilerResul
 
 class KISMETCOMPILER_API FKismetCompilerContext : public FGraphCompilerContext
 {
+public:
+
+	DECLARE_EVENT_OneParam(FKismetCompilerContext, FOnFunctionListCompiled, FKismetCompilerContext*);
+
 protected:
 	typedef FGraphCompilerContext Super;
 
@@ -61,6 +65,13 @@ protected:
 
 	// List of functions currently allocated
 	TIndirectArray<FKismetFunctionContext> FunctionList;
+
+	/** Set of function graphs generated for the class layout at compile time  */
+	TArray<UEdGraph*> GeneratedFunctionGraphs;
+
+	/** Event that is broadcast immediately after the function list for this context has been compiled. */
+	FOnFunctionListCompiled FunctionListCompiledEvent;
+
 protected:
 	// This struct holds the various compilation options, such as which passes to perform, whether to save intermediate results, etc
 	FKismetCompilerOptions CompileOptions;
@@ -134,6 +145,9 @@ public:
 	static FSimpleMulticastDelegate OnPreCompile;
 	static FSimpleMulticastDelegate OnPostCompile;
 
+	/** Broadcasts a notification immediately after the function list for this context has been compiled. */
+	FOnFunctionListCompiled& OnFunctionListCompiled() { return FunctionListCompiledEvent; }
+
 public:
 	FKismetCompilerContext(UBlueprint* SourceSketch, FCompilerResultsLog& InMessageLog, const FKismetCompilerOptions& InCompilerOptions);
 	virtual ~FKismetCompilerContext();
@@ -155,6 +169,9 @@ public:
 
 	const UEdGraphSchema_K2* GetSchema() const { return Schema; }
 
+	/** Spawn an intermediate function graph for this compilation using the specified desired name (which may be modified to make it unique */
+	UEdGraph* SpawnIntermediateFunctionGraph(const FString& InDesiredFunctionName);
+
 	// Spawns an intermediate node associated with the source node (for error purposes)
 	template <typename NodeType>
 	NodeType* SpawnIntermediateNode(UEdGraphNode* SourceNode, UEdGraph* ParentGraph = NULL)
diff --git a/Engine/Source/Runtime/Engine/Classes/Engine/Blueprint.h b/Engine/Source/Runtime/Engine/Classes/Engine/Blueprint.h
--- a/Engine/Source/Runtime/Engine/Classes/Engine/Blueprint.h
+++ b/Engine/Source/Runtime/Engine/Classes/Engine/Blueprint.h
@@ -20,6 +20,7 @@ class FCompilerResultsLog;
 class ITargetPlatform;
 class UActorComponent;
 class UEdGraph;
+class FKismetCompilerContext;
 class UInheritableComponentHandler;
 class FBlueprintActionDatabaseRegistrar;
 struct FDiffResults;
@@ -139,6 +140,18 @@ struct FBlueprintMacroCosmeticInfo
 	}
 };
 
+/** Structure that is passed to UBlueprint::GenerateFunctionGraphsEvent in order to generate procedural blueprint functions */
+USTRUCT()
+struct FGenerateBlueprintFunctionParams
+{
+	GENERATED_BODY()
+
+	/** The kismet compiler being used for compilation */
+	FKismetCompilerContext* CompilerContext;
+};
+
+DECLARE_DYNAMIC_MULTICAST_DELEGATE_OneParam(FGenerateBlueprintFunctionsEvent, const FGenerateBlueprintFunctionParams&, Params);
+
 struct FKismetCompilerOptions
 {
 public:
@@ -547,6 +560,11 @@ public:
 	/** Guid key for finding searchable data for Blueprint in the DDC */
 	UPROPERTY()
 	FGuid SearchGuid;
+
+	/** Persistent, serialized event that is triggered after skeleton class generation, but before functions are created for the class layout to give external subscribers an opportunity to generate function graphs on the compiler context */
+	UPROPERTY()
+	FGenerateBlueprintFunctionsEvent GenerateFunctionGraphsEvent;
+
 #endif //WITH_EDITORONLY_DATA
 
 	/** The version of the blueprint system that was used to  create this blueprint */
