Fixing CleanAndSanitizeStruct not destroying struct properties

diff --git a/Engine/Source/Editor/KismetCompiler/Private/UserDefinedStructureCompilerUtils.cpp b/Engine/Source/Editor/KismetCompiler/Private/UserDefinedStructureCompilerUtils.cpp
--- a/Engine/Source/Editor/KismetCompiler/Private/UserDefinedStructureCompilerUtils.cpp
+++ b/Engine/Source/Editor/KismetCompiler/Private/UserDefinedStructureCompilerUtils.cpp
@@ -159,13 +159,10 @@ struct FUserDefinedStructureCompilerInner
 
 			StructToClean->SetSuperStruct(nullptr);
 			StructToClean->Children = nullptr;
+			StructToClean->DestroyChildPropertiesAndResetPropertyLinks();
 			StructToClean->Script.Empty();
 			StructToClean->MinAlignment = 0;
-			StructToClean->RefLink = nullptr;
-			StructToClean->PropertyLink = nullptr;
-			StructToClean->DestructorLink = nullptr;
 			StructToClean->ScriptObjectReferences.Empty();
-			StructToClean->PropertyLink = nullptr;
 			StructToClean->ErrorMessage.Empty();
 			StructToClean->SetStructTrashed(true);
 		}
diff --git a/Engine/Source/Runtime/CoreUObject/Private/UObject/Class.cpp b/Engine/Source/Runtime/CoreUObject/Private/UObject/Class.cpp
--- a/Engine/Source/Runtime/CoreUObject/Private/UObject/Class.cpp
+++ b/Engine/Source/Runtime/CoreUObject/Private/UObject/Class.cpp
@@ -1573,6 +1573,15 @@ inline void DestroyPropertyLinkedList(FField*& PropertiesToDestroy)
 	PropertiesToDestroy = nullptr;
 }
 
+void UStruct::DestroyChildPropertiesAndResetPropertyLinks()
+{
+	DestroyPropertyLinkedList(ChildProperties);
+	PropertyLink = nullptr;
+	RefLink = nullptr;
+	DestructorLink = nullptr;
+	PostConstructLink = nullptr;
+}
+
 UStruct::~UStruct()
 {
 	// Destroy all properties owned by this struct
diff --git a/Engine/Source/Runtime/CoreUObject/Public/UObject/Class.h b/Engine/Source/Runtime/CoreUObject/Public/UObject/Class.h
--- a/Engine/Source/Runtime/CoreUObject/Public/UObject/Class.h
+++ b/Engine/Source/Runtime/CoreUObject/Public/UObject/Class.h
@@ -503,6 +503,9 @@ public:
 		return false;
 	}
 
+	/** Destroys all properties owned by this struct */
+	void DestroyChildPropertiesAndResetPropertyLinks();
+
 #if WITH_EDITORONLY_DATA
 	/** Try and find boolean metadata with the given key. If not found on this class, work up hierarchy looking for it. */
 	bool GetBoolMetaDataHierarchical(const FName& Key) const;
