Fix for a BP compilation error when a custom event is bound to a server-only delegate.

diff --git a/Engine/Source/Editor/KismetCompiler/Private/KismetCompiler.cpp b/Engine/Source/Editor/KismetCompiler/Private/KismetCompiler.cpp
--- a/Engine/Source/Editor/KismetCompiler/Private/KismetCompiler.cpp
+++ b/Engine/Source/Editor/KismetCompiler/Private/KismetCompiler.cpp
@@ -2941,7 +2941,18 @@ void FKismetCompilerContext::CreateFunctionStubForEvent(UK2Node_Event* SrcEventN
 	EntryNode->FunctionReference = SrcEventNode->EventReference;
 	EntryNode->CustomGeneratedFunctionName = EventNodeName;
 
-	if (!SrcEventNode->bOverrideFunction && SrcEventNode->IsUsedByAuthorityOnlyDelegate())
+	// Resolve expansions to original custom event node before checking it for a server-only delegate association.
+	auto IsServerOnlyEvent = [&MessageLog = this->MessageLog](UK2Node_Event* TargetEventNode)
+	{
+		if (UK2Node_CustomEvent* CustomEventNode = Cast<UK2Node_CustomEvent>(MessageLog.FindSourceObject(TargetEventNode)))
+		{
+			TargetEventNode = CustomEventNode;
+		}
+
+		return TargetEventNode->IsUsedByAuthorityOnlyDelegate();
+	};
+
+	if (!SrcEventNode->bOverrideFunction && IsServerOnlyEvent(SrcEventNode))
 	{
 		EntryNode->AddExtraFlags(FUNC_BlueprintAuthorityOnly);
 	}
