Suppress bogus warning from PVS Studio

#rnx

diff --git a/Engine/Source/Editor/KismetCompiler/Private/KismetCompiler.cpp b/Engine/Source/Editor/KismetCompiler/Private/KismetCompiler.cpp
--- a/Engine/Source/Editor/KismetCompiler/Private/KismetCompiler.cpp
+++ b/Engine/Source/Editor/KismetCompiler/Private/KismetCompiler.cpp
@@ -4660,6 +4660,7 @@ const UK2Node_FunctionEntry* FKismetCompilerContext::FindLocalEntryPoint(const U
 	return NULL;
 }
 
+#ifndef PVS_STUDIO // Bogus warning using GET_FUNCTION_NAME_CHECKED (see UE-88111)
 void FKismetCompilerContext::SetCanEverTick() const
 {
 	FTickFunction* TickFunction = nullptr;
@@ -4762,6 +4763,7 @@ void FKismetCompilerContext::SetCanEverTick() const
 			TickFunction->bCanEverTick ? *(CoreTexts.True.ToString()) : *(CoreTexts.False.ToString()) );
 	}
 }
+#endif
 
 bool FKismetCompilerContext::UsePersistentUberGraphFrame() const
 {
