Fix fast-path being incorrectly disabled for array properties

Fixes many animation issues. Blend nodes for a start don't function correctly without this fix.
Highlighted by this UDN: https://udn.unrealengine.com/questions/556298/425-animlayers-variable-value-is-not-updated-corre.html

diff --git a/Engine/Source/Editor/KismetCompiler/Private/AnimBlueprintCompiler.cpp b/Engine/Source/Editor/KismetCompiler/Private/AnimBlueprintCompiler.cpp
--- a/Engine/Source/Editor/KismetCompiler/Private/AnimBlueprintCompiler.cpp
+++ b/Engine/Source/Editor/KismetCompiler/Private/AnimBlueprintCompiler.cpp
@@ -3238,11 +3238,17 @@ void FAnimBlueprintCompilerContext::FPropertyCopyRecord::ValidateFastPath(UClass
 
 			// If the property is not one of the following types or categories, then we ignore the fast
 			// path. This matches the supported copyable types in FExposedValueHandler::Initialize
-			if (!(CastField<FBoolProperty>(DestProperty) ||
-				  CastField<FNameProperty>(DestProperty) ||
-				  CastField<FStructProperty>(DestProperty) ||
-				  CastField<FObjectPropertyBase>(DestProperty) ||
-				  (DestProperty->PropertyFlags & CPF_IsPlainOldData) != 0))
+			FProperty* DestPropertyOrInner = DestProperty;
+			if (FArrayProperty* DestArrayProperty = CastField<FArrayProperty>(DestProperty))
+			{
+				DestPropertyOrInner = DestArrayProperty->Inner;
+			}
+
+			if (!(CastField<FBoolProperty>(DestPropertyOrInner) ||
+				  CastField<FNameProperty>(DestPropertyOrInner) ||
+				  CastField<FStructProperty>(DestPropertyOrInner) ||
+				  CastField<FObjectPropertyBase>(DestPropertyOrInner) ||
+				  (DestPropertyOrInner->PropertyFlags & CPF_IsPlainOldData) != 0))
 			{
 				InvalidateFastPath();
 				return;
