Fix split pins not working correctly in animation state machine states

Added calls to split the pins of all nodes when states are merged into the consolidated event graph

diff --git a/Engine/Source/Editor/BlueprintGraph/Classes/K2Node.h b/Engine/Source/Editor/BlueprintGraph/Classes/K2Node.h
--- a/Engine/Source/Editor/BlueprintGraph/Classes/K2Node.h
+++ b/Engine/Source/Editor/BlueprintGraph/Classes/K2Node.h
@@ -406,6 +406,9 @@ class UK2Node : public UEdGraphNode
 	/** Return whether this node references the specified variable, give the supplied scope. Used when variable types are changed. */
 	virtual bool ReferencesVariable(const FName& InVarName, const UStruct* InScope) const { return false; }
 
+	/** Helper function for ExpandNode(), allowing other contexts to call pin expansion alone */
+	BLUEPRINTGRAPH_API void ExpandSplitPins(FKismetCompilerContext& CompilerContext, UEdGraph* SourceGraph);
+
 protected:
 
 	enum ERedirectType
diff --git a/Engine/Source/Editor/BlueprintGraph/Private/K2Node.cpp b/Engine/Source/Editor/BlueprintGraph/Private/K2Node.cpp
--- a/Engine/Source/Editor/BlueprintGraph/Private/K2Node.cpp
+++ b/Engine/Source/Editor/BlueprintGraph/Private/K2Node.cpp
@@ -1280,6 +1280,11 @@ UK2Node* UK2Node::ExpandSplitPin(FKismetCompilerContext* CompilerContext, UEdGra
 }
 
 void UK2Node::ExpandNode(FKismetCompilerContext& CompilerContext, UEdGraph* SourceGraph)
+{
+	ExpandSplitPins(CompilerContext, SourceGraph);
+}
+
+void UK2Node::ExpandSplitPins(FKismetCompilerContext& CompilerContext, UEdGraph* SourceGraph)
 {
 	// We iterate the array in reverse so we can both remove the subpins safely after we've read them and
 	// so we have split nested structs we combine them back together in the right order
diff --git a/Engine/Source/Editor/KismetCompiler/Private/AnimBlueprintCompiler.cpp b/Engine/Source/Editor/KismetCompiler/Private/AnimBlueprintCompiler.cpp
--- a/Engine/Source/Editor/KismetCompiler/Private/AnimBlueprintCompiler.cpp
+++ b/Engine/Source/Editor/KismetCompiler/Private/AnimBlueprintCompiler.cpp
@@ -1315,14 +1315,16 @@ void FAnimBlueprintCompilerContext::ProcessAllAnimationNodes()
 		MessageLog.Error(*LOCTEXT("ExpectedAFunctionEntry_Error", "Expected at least one animation root, but did not find any").ToString());
 	}
 
-		// Build cached pose map
-		BuildCachedPoseNodeUpdateOrder();
+	// Build cached pose map
+	BuildCachedPoseNodeUpdateOrder();
 }
 
 int32 FAnimBlueprintCompilerContext::ExpandGraphAndProcessNodes(UEdGraph* SourceGraph, UAnimGraphNode_Base* SourceRootNode, UAnimStateTransitionNode* TransitionNode, TArray<UEdGraphNode*>* ClonedNodes)
 {
 	// Clone the nodes from the source graph
-	UEdGraph* ClonedGraph = FEdGraphUtilities::CloneGraph(SourceGraph, NULL, &MessageLog, true);
+	// Note that we outer this graph to the ConsolidatedEventGraph to allow ExpandSplitPins to 
+	// correctly retrieve the context for any expanded function calls (custom make/break structs etc.)
+	UEdGraph* ClonedGraph = FEdGraphUtilities::CloneGraph(SourceGraph, ConsolidatedEventGraph, &MessageLog, true);
 
 	// Grab all the animation nodes and find the corresponding root node in the cloned set
 	UAnimGraphNode_Base* TargetRootNode = NULL;
@@ -1360,6 +1362,9 @@ int32 FAnimBlueprintCompilerContext::ExpandGraphAndProcessNodes(UEdGraph* Source
 	}
 	check(TargetRootNode);
 
+	// Run another expansion pass to catch the graph we just added (this is slightly wasteful 
+	ExpandSplitPins(ClonedGraph);
+
 	// Move the cloned nodes into the consolidated event graph
 	const bool bIsLoading = Blueprint->bIsRegeneratingOnLoad || IsAsyncLoading();
 	const bool bIsCompiling = Blueprint->bBeingCompiled;
@@ -2103,6 +2108,18 @@ void FAnimBlueprintCompilerContext::CopyTermDefaultsToDefaultObject(UObject* Def
 	}
 }
 
+void FAnimBlueprintCompilerContext::ExpandSplitPins(UEdGraph* InGraph)
+{
+	for (TArray<UEdGraphNode*>::TIterator NodeIt(InGraph->Nodes); NodeIt; ++NodeIt)
+	{
+		UK2Node* K2Node = Cast<UK2Node>(*NodeIt);
+		if (K2Node != nullptr)
+		{
+			K2Node->ExpandSplitPins(*this, InGraph);
+		}
+	}
+}
+
 // Merges in any all ubergraph pages into the gathering ubergraph
 void FAnimBlueprintCompilerContext::MergeUbergraphPagesIn(UEdGraph* Ubergraph)
 {
@@ -2141,24 +2158,7 @@ void FAnimBlueprintCompilerContext::MergeUbergraphPagesIn(UEdGraph* Ubergraph)
 		}
 
 		// Make sure we expand any split pins here before we process animation nodes.
-		for (TArray<UEdGraphNode*>::TIterator NodeIt(ConsolidatedEventGraph->Nodes); NodeIt; ++NodeIt)
-		{
-			UK2Node* K2Node = Cast<UK2Node>(*NodeIt);
-			if (K2Node != nullptr)
-			{
-				// We iterate the array in reverse so we can recombine split-pins (which modifies the pins array)
-				for (int32 PinIndex = K2Node->Pins.Num() - 1; PinIndex >= 0; --PinIndex)
-				{
-					UEdGraphPin* Pin = K2Node->Pins[PinIndex];
-					if (Pin->SubPins.Num() == 0)
-					{
-						continue;
-					}
-
-					K2Node->ExpandSplitPin(this, ConsolidatedEventGraph, Pin);
-				}
-			}
-		}
+		ExpandSplitPins(ConsolidatedEventGraph);
 
 		// Compile the animation graph
 		ProcessAllAnimationNodes();
diff --git a/Engine/Source/Editor/KismetCompiler/Private/AnimBlueprintCompiler.h b/Engine/Source/Editor/KismetCompiler/Private/AnimBlueprintCompiler.h
--- a/Engine/Source/Editor/KismetCompiler/Private/AnimBlueprintCompiler.h
+++ b/Engine/Source/Editor/KismetCompiler/Private/AnimBlueprintCompiler.h
@@ -378,5 +378,8 @@ private:
 
 	// Clean up transient stub functions
 	void DestroyAnimGraphStubFunctions();
+
+	// Expands split pins for a graph
+	void ExpandSplitPins(UEdGraph* InGraph);
 };
 
