Add Kismet compiler validation for the default value on  class reference type variables

diff --git a/Engine/Source/Editor/Kismet/Private/BlueprintCompilationManager.cpp b/Engine/Source/Editor/Kismet/Private/BlueprintCompilationManager.cpp
--- a/Engine/Source/Editor/Kismet/Private/BlueprintCompilationManager.cpp
+++ b/Engine/Source/Editor/Kismet/Private/BlueprintCompilationManager.cpp
@@ -440,7 +440,7 @@ struct FCompilerData
 	bool IsCppCompileType() const { return InternalOptions.CompileType == EKismetCompileType::Cpp; }
 	bool ShouldSetTemporaryBlueprintFlags() const { return JobType != ECompilationManagerJobType::RelinkOnly; }
 	bool ShouldResetErrorState() const { return JobType == ECompilationManagerJobType::Normal && InternalOptions.CompileType != EKismetCompileType::BytecodeOnly; }
-	bool ShouldValidateVariableNames() const { return JobType == ECompilationManagerJobType::Normal; }
+	bool ShouldValidate() const { return JobType == ECompilationManagerJobType::Normal; }
 	bool ShouldRegenerateSkeleton() const { return JobType != ECompilationManagerJobType::RelinkOnly; }
 	bool ShouldMarkUpToDateAfterSkeletonStage() const { return IsSkeletonOnly(); }
 	bool ShouldReconstructNodes() const { return JobType == ECompilationManagerJobType::Normal; }
@@ -836,15 +836,16 @@ void FBlueprintCompilationManagerImpl::FlushCompilationQueueImpl(bool bSuppressB
 			}
 		}
 
-		// STAGE V: Validate Variable Names
+		// STAGE V: Validate
 		for (FCompilerData& CompilerData : CurrentlyCompilingBPs)
 		{
-			if(!CompilerData.ShouldValidateVariableNames())
+			if(!CompilerData.ShouldValidate())
 			{
 				continue;
 			}
 
 			CompilerData.Compiler->ValidateVariableNames();
+			CompilerData.Compiler->ValidateClassPropertyDefaults();
 		}
 
 		// STAGE VI: Purge null graphs, misc. data fixup
diff --git a/Engine/Source/Editor/KismetCompiler/Private/KismetCompiler.cpp b/Engine/Source/Editor/KismetCompiler/Private/KismetCompiler.cpp
--- a/Engine/Source/Editor/KismetCompiler/Private/KismetCompiler.cpp
+++ b/Engine/Source/Editor/KismetCompiler/Private/KismetCompiler.cpp
@@ -634,6 +634,53 @@ void FKismetCompilerContext::ValidateComponentClassOverrides()
 	}
 }
 
+void FKismetCompilerContext::ValidateClassPropertyDefaults()
+{
+	// We cannot check the default value if there is no generated class yet (i.e. opening a BP for the first time)
+	if (!Blueprint->GeneratedClass)
+	{
+		return;
+	}
+
+	UObject* CDO = Blueprint->GeneratedClass->GetDefaultObject(false);
+	if (!CDO)
+	{
+		return;
+	}
+
+	// Check the New Variables in case the user has changed a type since the last check and a newer
+	// CDO has not been generated yet
+	for (FBPVariableDescription& VarDesc : Blueprint->NewVariables)
+	{
+		if (VarDesc.VarType.PinCategory == UEdGraphSchema_K2::PC_Class)
+		{
+			const UObject* ExpectedObject = VarDesc.VarType.PinSubCategoryObject.Get();
+			const UClass* ExpectedClass = Cast<UClass>(ExpectedObject);
+
+			if (FClassProperty* AsClassProperty = FindFProperty<FClassProperty>(Blueprint->GeneratedClass, VarDesc.VarName))
+			{
+				if (UObject* SetDefaultProp = AsClassProperty->GetObjectPropertyValue_InContainer(CDO))
+				{
+					// Object may be of a type that is also being compiled and therefore REINST_, so get the real class
+					const UClass* SetDefaultClass = Cast<UClass>(SetDefaultProp);
+					const UClass* RealClass = SetDefaultClass ? SetDefaultClass->GetAuthoritativeClass() : nullptr;
+
+					if (RealClass && ExpectedClass && !RealClass->IsChildOf(ExpectedClass))
+					{
+						MessageLog.Error(
+							*FText::Format(
+								LOCTEXT("InvalidClassPropertyDefaultValue_ErrorFmt", "Invalid default type '{0}' on property '{1}' in @@"),
+								FText::FromString(RealClass->GetName()), FText::FromString(AsClassProperty->GetName())
+							).ToString(),
+							Blueprint
+						);
+					}
+				}
+			}
+		}
+	}
+}
+
 void FKismetCompilerContext::ValidateTimelineNames()
 {
 	TSharedPtr<FKismetNameValidator> ParentBPNameValidator;
diff --git a/Engine/Source/Editor/KismetCompiler/Public/KismetCompiler.h b/Engine/Source/Editor/KismetCompiler/Public/KismetCompiler.h
--- a/Engine/Source/Editor/KismetCompiler/Public/KismetCompiler.h
+++ b/Engine/Source/Editor/KismetCompiler/Public/KismetCompiler.h
@@ -288,6 +288,13 @@ public:
 	/** Ensures that all component class overrides are legal overrides of the parent class */
 	void ValidateComponentClassOverrides();
 
+	/**
+	* Ensures that all class reference Properties are legal overrides of the parent class
+	* by checking the default value set on any PC_Class variable types. Requires a 
+	* valid CDO in order to do this validation. Called in Stage V: Validate 
+	*/
+	void ValidateClassPropertyDefaults();
+
 	/** Creates a class variable for each entry in the Blueprint NewVars array */
 	virtual void CreateClassVariablesFromBlueprint();
 
