Blueprints: Allow function calls to native methods with return values that are a const reference to be left unwired without generating a nonsensical 'literals are not allowed' warning

diff --git a/Engine/Source/Editor/KismetCompiler/Private/KismetCompilerMisc.cpp b/Engine/Source/Editor/KismetCompiler/Private/KismetCompilerMisc.cpp
--- a/Engine/Source/Editor/KismetCompiler/Private/KismetCompilerMisc.cpp
+++ b/Engine/Source/Editor/KismetCompiler/Private/KismetCompilerMisc.cpp
@@ -431,7 +431,8 @@ bool FKismetCompilerUtilities::IsTypeCompatibleWithProperty(UEdGraphPin* SourceP
 
 		if (Property->HasAnyPropertyFlags(CPF_ReferenceParm)
 			&& (SourcePin->LinkedTo.Num() == 0)
-			&& (SourcePin->PinType.PinSubCategoryObject.Get() != TBaseStructure<FTransform>::Get()))
+			&& (SourcePin->PinType.PinSubCategoryObject.Get() != TBaseStructure<FTransform>::Get())
+			&& (SourcePin->Direction == EGPD_Input))
 		{
 			TArray<FString> AutoEmittedTerms;
 			Schema->GetAutoEmitTermParameters(OwningFunction, AutoEmittedTerms);
