Add Insights traces to the Kismet Compiler

#rnx

diff --git a/Engine/Source/Editor/KismetCompiler/Private/KismetCompiler.cpp b/Engine/Source/Editor/KismetCompiler/Private/KismetCompiler.cpp
--- a/Engine/Source/Editor/KismetCompiler/Private/KismetCompiler.cpp
+++ b/Engine/Source/Editor/KismetCompiler/Private/KismetCompiler.cpp
@@ -176,6 +176,7 @@ FKismetCompilerContext::~FKismetCompilerContext()
 
 UEdGraphSchema_K2* FKismetCompilerContext::CreateSchema()
 {
+	TRACE_CPUPROFILER_EVENT_SCOPE(WILD_FKismetCompilerContext::CreateSchema);
 	return NewObject<UEdGraphSchema_K2>();
 }
 
@@ -361,6 +362,7 @@ void FKismetCompilerContext::SaveSubObjectsFromCleanAndSanitizeClass(FSubobjectC
 
 void FKismetCompilerContext::PostCreateSchema()
 {
+	TRACE_CPUPROFILER_EVENT_SCOPE(WILD_FKismetCompilerContext::PostCreateSchema);
 	NodeHandlers.Add(UEdGraphNode_Comment::StaticClass(), new FNodeHandlingFunctor(*this));
 
 	TArray<UClass*> ClassesOfUK2Node;
@@ -3886,6 +3888,7 @@ void FKismetCompilerContext::ValidateFunctionGraphNames()
 // Creates a copy of the graph to allow further transformations to occur
 void FKismetCompilerContext::CreateFunctionList()
 {
+	TRACE_CPUPROFILER_EVENT_SCOPE(WILD_FKismetCompilerContext::CreateFunctionList);
 	{
 		BP_SCOPED_COMPILER_EVENT_STAT(EKismetCompilerStats_GenerateFunctionGraphs);
 
