Do not attempt to get authoritative class on a trash struct during compilation because it will crash

#rnx

diff --git a/Engine/Source/Editor/KismetCompiler/Private/KismetCompilerMisc.cpp b/Engine/Source/Editor/KismetCompiler/Private/KismetCompilerMisc.cpp
--- a/Engine/Source/Editor/KismetCompiler/Private/KismetCompilerMisc.cpp
+++ b/Engine/Source/Editor/KismetCompiler/Private/KismetCompilerMisc.cpp
@@ -158,14 +158,27 @@ static bool DoesTypeNotMatchProperty(UEdGraphPin* SourcePin, const FEdGraphPinTy
 	{
 		const UClass* ObjectType = (PinSubCategory == UEdGraphSchema_K2::PSC_Self) ? SelfClass : Cast<const UClass>(PinSubCategoryObject);
 
-		if (ObjectType == NULL)
+
+		if (!ObjectType)
 		{
 			MessageLog.Error(*LOCTEXT("FindClassForPin_Error", "Failed to find class for pin @@").ToString(), SourcePin);
 		}
+		// If the object type has been marked as transient and is no longer rooted in the GUObjectArray,
+		// then then it has been "consigned to oblvion". This can be the case if a BP asset has been force 
+		// deleted and references to it are still laying around
+		else if(
+			ObjectType->HasAnyFlags(RF_Transient) &&
+			ObjectType->HasAnyClassFlags(CLASS_NewerVersionExists) &&
+			!ObjectType->IsRooted() && 
+			ObjectType->ClassDefaultObject &&
+			!ObjectType->ClassDefaultObject->IsRooted())
+		{
+			MessageLog.Error(*LOCTEXT("InvalidClassForPin_Error", "Class for pin @@ is invalid! It has likely been deleted.").ToString(), SourcePin);
+		}
 		else
 		{
 			FObjectPropertyBase* ObjProperty = CastField<FObjectPropertyBase>(TestProperty);
-			if (ObjProperty != NULL && ObjProperty->PropertyClass)
+			if (ObjProperty != nullptr && ObjProperty->PropertyClass)
 			{
 				const UClass* OutputClass = (Direction == EGPD_Output) ? ObjectType : ObjProperty->PropertyClass;
 				const UClass* InputClass  = (Direction == EGPD_Output) ? ObjProperty->PropertyClass : ObjectType;
@@ -210,7 +223,7 @@ static bool DoesTypeNotMatchProperty(UEdGraphPin* SourcePin, const FEdGraphPinTy
 			else if (FInterfaceProperty* IntefaceProperty = CastField<FInterfaceProperty>(TestProperty))
 			{
 				UClass const* InterfaceClass = IntefaceProperty->InterfaceClass;
-				if (InterfaceClass == NULL)
+				if (InterfaceClass == nullptr)
 				{
 					bTypeMismatch = true;
 				}
