Added some extra profiling scopes to kismet compiler etc.

diff --git a/Engine/Source/Editor/KismetCompiler/Private/KismetCompiler.cpp b/Engine/Source/Editor/KismetCompiler/Private/KismetCompiler.cpp
--- a/Engine/Source/Editor/KismetCompiler/Private/KismetCompiler.cpp
+++ b/Engine/Source/Editor/KismetCompiler/Private/KismetCompiler.cpp
@@ -70,6 +70,8 @@ FSimpleMulticastDelegate FKismetCompilerContext::OnPostCompile;
 DECLARE_CYCLE_STAT(TEXT("Create Schema"), EKismetCompilerStats_CreateSchema, STATGROUP_KismetCompiler );
 DECLARE_CYCLE_STAT(TEXT("Create Function List"), EKismetCompilerStats_CreateFunctionList, STATGROUP_KismetCompiler );
 DECLARE_CYCLE_STAT(TEXT("Expansion"), EKismetCompilerStats_Expansion, STATGROUP_KismetCompiler )
+DECLARE_CYCLE_STAT(TEXT("Expand Node"), EKismetCompilerStats_ExpandNode, STATGROUP_KismetCompiler )
+DECLARE_CYCLE_STAT(TEXT("Post Expansion Step"), EKismetCompilerStats_PostExpansionStep, STATGROUP_KismetCompiler )
 DECLARE_CYCLE_STAT(TEXT("Process uber"), EKismetCompilerStats_ProcessUbergraph, STATGROUP_KismetCompiler );
 DECLARE_CYCLE_STAT(TEXT("Process func"), EKismetCompilerStats_ProcessFunctionGraph, STATGROUP_KismetCompiler );
 DECLARE_CYCLE_STAT(TEXT("Generate Function Graph"), EKismetCompilerStats_GenerateFunctionGraphs, STATGROUP_KismetCompiler );
@@ -84,6 +86,8 @@ DECLARE_CYCLE_STAT(TEXT("Bind and Link Class"), EKismetCompilerStats_BindAndLink
 DECLARE_CYCLE_STAT(TEXT("Calculate checksum of CDO"), EKismetCompilerStats_ChecksumCDO, STATGROUP_KismetCompiler );
 DECLARE_CYCLE_STAT(TEXT("Analyze execution path"), EKismetCompilerStats_AnalyzeExecutionPath, STATGROUP_KismetCompiler);
 DECLARE_CYCLE_STAT(TEXT("Calculate checksum of signature"), EKismetCompilerStats_ChecksumSignature, STATGROUP_KismetCompiler);
+DECLARE_CYCLE_STAT(TEXT("Pruning"), EKismetCompilerStats_PruneIsolatedNodes, STATGROUP_KismetCompiler);
+DECLARE_CYCLE_STAT(TEXT("Merge Ubergraph Pages In"), EKismetCompilerStats_MergeUbergraphPagesIn, STATGROUP_KismetCompiler);
 
 namespace
 {
@@ -1484,6 +1488,8 @@ void FKismetCompilerContext::PruneIsolatedNodes(UEdGraph* InGraph, bool bInInclu
 /** Prunes any nodes that weren't visited from the graph, printing out a warning */
 void FKismetCompilerContext::PruneIsolatedNodes(const TArray<UEdGraphNode*>& RootSet, TArray<UEdGraphNode*>& GraphNodes)
 {
+	BP_SCOPED_COMPILER_EVENT_STAT(EKismetCompilerStats_PruneIsolatedNodes);
+	
 	//@TODO: This function crawls the graph twice (once here and once in Super, could potentially combine them, with a bitflag for flows reached via exec wires)
 
 	// Prune the impure nodes that aren't reachable via any (even impossible, e.g., a branch never taken) execution flow
@@ -3101,6 +3107,8 @@ void FKismetCompilerContext::CreateFunctionStubForEvent(UK2Node_Event* SrcEventN
 
 void FKismetCompilerContext::MergeUbergraphPagesIn(UEdGraph* Ubergraph)
 {
+	BP_SCOPED_COMPILER_EVENT_STAT(EKismetCompilerStats_MergeUbergraphPagesIn);
+	
 	for (TArray<UEdGraph*>::TIterator It(Blueprint->UbergraphPages); It; ++It)
 	{
 		UEdGraph* SourceGraph = *It;
@@ -3160,6 +3168,7 @@ void FKismetCompilerContext::ExpansionStep(UEdGraph* Graph, bool bAllowUbergraph
 			UK2Node_Knot* KnotNode = Cast<UK2Node_Knot>(Graph->Nodes[NodeIndex]);
 			if (KnotNode)
 			{
+				BP_SCOPED_COMPILER_EVENT_STAT(EKismetCompilerStats_ExpandNode);
 				KnotNode->ExpandNode(*this, Graph);
 			}
 		}
@@ -3169,6 +3178,7 @@ void FKismetCompilerContext::ExpansionStep(UEdGraph* Graph, bool bAllowUbergraph
 			UK2Node* Node = Cast<UK2Node>(Graph->Nodes[NodeIndex]);
 			if (Node)
 			{
+				BP_SCOPED_COMPILER_EVENT_STAT(EKismetCompilerStats_ExpandNode);
 				Node->ExpandNode(*this, Graph);
 			}
 		}
@@ -3184,7 +3194,10 @@ void FKismetCompilerContext::ExpansionStep(UEdGraph* Graph, bool bAllowUbergraph
 		ExpandTimelineNodes(Graph);
 	}
 
-	PostExpansionStep(Graph);
+	{
+		BP_SCOPED_COMPILER_EVENT_STAT(EKismetCompilerStats_PostExpansionStep);
+		PostExpansionStep(Graph);
+	}
 }
 
 void FKismetCompilerContext::DetermineNodeExecLinks(UEdGraphNode* SourceNode, TMap<UEdGraphPin*, UEdGraphPin*>& SourceNodeLinks) const
diff --git a/Engine/Source/Editor/UnrealEd/Private/EdGraphUtilities.cpp b/Engine/Source/Editor/UnrealEd/Private/EdGraphUtilities.cpp
--- a/Engine/Source/Editor/UnrealEd/Private/EdGraphUtilities.cpp
+++ b/Engine/Source/Editor/UnrealEd/Private/EdGraphUtilities.cpp
@@ -15,6 +15,9 @@
 #include "Widgets/Notifications/SNotificationList.h"
 #include "K2Node_Composite.h"
 
+DECLARE_CYCLE_STAT(TEXT("Clone Graph"), EKismetCompilerStats_CloneGraph, STATGROUP_KismetCompiler);
+DECLARE_CYCLE_STAT(TEXT("Clone Graph - Build Backtrack Map"), EKismetCompilerStats_CloneGraph_BuildBackTrackMap, STATGROUP_KismetCompiler);
+
 /////////////////////////////////////////////////////
 // Local namespace
 
@@ -238,6 +241,8 @@ UEdGraphPin* FEdGraphUtilities::GetNetFromPin(UEdGraphPin* Pin)
 // maintaining a mapping from the clone to the source nodes (even across multiple clonings)
 UEdGraph* FEdGraphUtilities::CloneGraph(UEdGraph* InSource, UObject* NewOuter, FCompilerResultsLog* MessageLog, bool bCloningForCompile)
 {
+	BP_SCOPED_COMPILER_EVENT_STAT(EKismetCompilerStats_CloneGraph);
+
 	// Duplicate the graph, keeping track of what was duplicated
 	TMap<UObject*, UObject*> DuplicatedObjectList;
 
@@ -256,6 +261,8 @@ UEdGraph* FEdGraphUtilities::CloneGraph(UEdGraph* InSource, UObject* NewOuter, F
 	// Store backtrack links from each duplicated object to the original source object
 	if (MessageLog != NULL)
 	{
+		BP_SCOPED_COMPILER_EVENT_STAT(EKismetCompilerStats_CloneGraph_BuildBackTrackMap);
+		
 		for (TMap<UObject*, UObject*>::TIterator It(DuplicatedObjectList); It; ++It)
 		{
 			UObject* const Source = It.Key();
diff --git a/Engine/Source/Editor/UnrealEd/Public/Kismet2/CompilerResultsLog.h b/Engine/Source/Editor/UnrealEd/Public/Kismet2/CompilerResultsLog.h
--- a/Engine/Source/Editor/UnrealEd/Public/Kismet2/CompilerResultsLog.h
+++ b/Engine/Source/Editor/UnrealEd/Public/Kismet2/CompilerResultsLog.h
@@ -440,9 +440,11 @@ public:
 #if STATS
 #define BP_SCOPED_COMPILER_EVENT_STAT(Stat) \
 	SCOPE_CYCLE_COUNTER(Stat); \
+	TRACE_CPUPROFILER_EVENT_SCOPE(Stat); \
 	FScopedCompilerEvent PREPROCESSOR_JOIN(ScopedCompilerEvent,__LINE__)(GET_STATDESCRIPTION(Stat))
 #else
 #define BP_SCOPED_COMPILER_EVENT_STAT(Stat) \
+	TRACE_CPUPROFILER_EVENT_SCOPE(Stat); \
 	FScopedCompilerEvent PREPROCESSOR_JOIN(ScopedCompilerEvent,__LINE__)(ANSI_TO_TCHAR(#Stat))
 #endif
 #endif	//#if WITH_EDITOR
