Added the concept of thread safe Blueprint functions

Either enabling the 'thread safe' checkbox for a function or overriding an already thread safe (either native or BP) function will enable compile-time checking for thread safety.
Checks are based around the complied statement list rather than the bytecode itself.

diff --git a/Engine/Source/Editor/BlueprintGraph/Classes/EdGraphSchema_K2.h b/Engine/Source/Editor/BlueprintGraph/Classes/EdGraphSchema_K2.h
--- a/Engine/Source/Editor/BlueprintGraph/Classes/EdGraphSchema_K2.h
+++ b/Engine/Source/Editor/BlueprintGraph/Classes/EdGraphSchema_K2.h
@@ -250,6 +250,12 @@ public:
 	/** Namespace into which a type can be optionally defined; if empty or not set, the type will belong to the global namespace (default). */
 	static const FName MD_Namespace;
 
+	/** Function or class marked as thread-safe. Opts class/function compilation into thread-safety checks. */
+	static const FName MD_ThreadSafe;
+
+	/** Function marked as explicitly not thread-safe. Opts function out of class-level thread-safety checks. */
+	static const FName MD_NotThreadSafe;
+	
 private:
 	// This class should never be instantiated
 	FBlueprintMetadata() {}
@@ -1093,6 +1099,13 @@ public:
 	 */
 	bool DoesGraphSupportImpureFunctions(const UEdGraph* InGraph) const;
 
+	/** 
+	 * Checks if the graph is marked as thread safe
+	 * @param InGraph		Graph to check
+	 * @return			True if the graph is marked theead safe
+	 */
+	bool IsGraphMarkedThreadSafe(const UEdGraph* InGraph) const;
+	
 	/**
 	 * Checks to see if the passed in function is valid in the graph for the current class
 	 *
diff --git a/Engine/Source/Editor/BlueprintGraph/Classes/K2Node_EditablePinBase.h b/Engine/Source/Editor/BlueprintGraph/Classes/K2Node_EditablePinBase.h
--- a/Engine/Source/Editor/BlueprintGraph/Classes/K2Node_EditablePinBase.h
+++ b/Engine/Source/Editor/BlueprintGraph/Classes/K2Node_EditablePinBase.h
@@ -71,6 +71,9 @@ struct FKismetUserDeclaredFunctionMetadata
 	UPROPERTY()
 	bool bCallInEditor;
 
+	UPROPERTY()
+	bool bThreadSafe;
+
 	/** Cached value for whether or not the graph has latent functions, positive for TRUE, zero for FALSE, and INDEX_None for undetermined */
 	UPROPERTY(transient)
 	int8 HasLatentFunctions;
@@ -80,6 +83,7 @@ public:
 		: InstanceTitleColor(FLinearColor::White)
 		, bIsDeprecated(false)
 		, bCallInEditor(false)
+		, bThreadSafe(false)
 		, HasLatentFunctions(INDEX_NONE)
 	{
 	}
diff --git a/Engine/Source/Editor/BlueprintGraph/Private/BlueprintActionFilter.cpp b/Engine/Source/Editor/BlueprintGraph/Private/BlueprintActionFilter.cpp
--- a/Engine/Source/Editor/BlueprintGraph/Private/BlueprintActionFilter.cpp
+++ b/Engine/Source/Editor/BlueprintGraph/Private/BlueprintActionFilter.cpp
@@ -184,6 +184,16 @@ namespace BlueprintActionFilterImpl
 	*/
 	static bool IsPropertyAccessorNode(FBlueprintActionFilter const& Filter, FBlueprintActionInfo& BlueprintAction);
 
+	/**
+	* Rejection test that checks to see if the supplied node-spawner would produce a node with thread safety
+	* characteristics that are incompatible with the specification of the containing graph. 
+	* 
+	* @param  Filter			Filter context (unused) for this test.
+	* @param  BlueprintAction	The action you wish to query.
+	* @return True if the action would spawn a node that has incompatible thread safety.
+	*/
+	static bool IsThreadSafetyIncompatible(FBlueprintActionFilter const& Filter, FBlueprintActionInfo& BlueprintAction);
+
 	/**
 	 * Rejection test that checks to see if the supplied node-spawner would 
 	 * produce an impure node, incompatible with the specified graphs.
@@ -998,6 +1008,32 @@ static bool BlueprintActionFilterImpl::IsPropertyAccessorNode(FBlueprintActionFi
 	return bIsAccessor;
 }
 
+//------------------------------------------------------------------------------
+static bool BlueprintActionFilterImpl::IsThreadSafetyIncompatible(FBlueprintActionFilter const& Filter, FBlueprintActionInfo& BlueprintAction)
+{
+	if (UFunction const* Function = BlueprintAction.GetAssociatedFunction())
+	{	
+		bool bAllowNonThreadSafeNodes = true;
+
+		FBlueprintActionContext const& FilterContext = Filter.Context;
+
+		for (UEdGraph* Graph : FilterContext.Graphs)
+		{
+			if (const UEdGraphSchema_K2* K2Schema = Cast<UEdGraphSchema_K2>(Graph->GetSchema()))
+			{
+				bAllowNonThreadSafeNodes &= !K2Schema->IsGraphMarkedThreadSafe(Graph);
+			}
+		}
+
+		if(!bAllowNonThreadSafeNodes)
+		{
+			return !FBlueprintEditorUtils::HasFunctionBlueprintThreadSafeMetaData(Function);
+		}
+	}
+	
+	return false;
+}
+
 //------------------------------------------------------------------------------
 static bool BlueprintActionFilterImpl::IsIncompatibleImpureNode(FBlueprintActionFilter const& Filter, FBlueprintActionInfo& BlueprintAction)
 {
@@ -1993,6 +2029,10 @@ FBlueprintActionFilter::FBlueprintActionFilter(uint32 Flags/*= 0x00*/)
 	AddRejectionTest(FRejectionTestDelegate::CreateStatic(IsIncompatibleLatentNode));
 	AddRejectionTest(FRejectionTestDelegate::CreateStatic(IsIncompatibleImpureNode));
 	AddRejectionTest(FRejectionTestDelegate::CreateStatic(IsPropertyAccessorNode));
+	if (Flags & BPFILTER_RejectIncompatibleThreadSafety)
+	{
+		AddRejectionTest(FRejectionTestDelegate::CreateStatic(IsThreadSafetyIncompatible));
+	}
 	
 	AddRejectionTest(FRejectionTestDelegate::CreateStatic(IsActionHiddenByConfig));
 	AddRejectionTest(FRejectionTestDelegate::CreateStatic(IsFieldCategoryHidden));
diff --git a/Engine/Source/Editor/BlueprintGraph/Private/EdGraphSchema_K2.cpp b/Engine/Source/Editor/BlueprintGraph/Private/EdGraphSchema_K2.cpp
--- a/Engine/Source/Editor/BlueprintGraph/Private/EdGraphSchema_K2.cpp
+++ b/Engine/Source/Editor/BlueprintGraph/Private/EdGraphSchema_K2.cpp
@@ -199,6 +199,9 @@ const FName FBlueprintMetadata::MD_AllowAbstractClasses(TEXT("AllowAbstract"));
 
 const FName FBlueprintMetadata::MD_Namespace(TEXT("Namespace"));
 
+const FName FBlueprintMetadata::MD_ThreadSafe(TEXT("BlueprintThreadSafe"));
+const FName FBlueprintMetadata::MD_NotThreadSafe(TEXT("NotBlueprintThreadSafe"));
+
 //////////////////////////////////////////////////////////////////////////
 
 #define LOCTEXT_NAMESPACE "KismetSchema"
@@ -971,8 +974,8 @@ bool UEdGraphSchema_K2::HasFunctionAnyOutputParameter(const UFunction* InFunctio
 
 bool UEdGraphSchema_K2::FunctionCanBePlacedAsEvent(const UFunction* InFunction)
 {
-	// First check we are override-able, non-static and non-const
-	if (!InFunction || !CanKismetOverrideFunction(InFunction) || InFunction->HasAnyFunctionFlags(FUNC_Static|FUNC_Const))
+	// First check we are override-able, non-static, non-const and not marked thread safe
+	if (!InFunction || !CanKismetOverrideFunction(InFunction) || InFunction->HasAnyFunctionFlags(FUNC_Static|FUNC_Const) || FBlueprintEditorUtils::HasFunctionBlueprintThreadSafeMetaData(InFunction))
 	{
 		return false;
 	}
@@ -1267,6 +1270,29 @@ bool UEdGraphSchema_K2::DoesGraphSupportImpureFunctions(const UEdGraph* InGraph)
 	return bAllowImpureFuncs;
 }
 
+bool UEdGraphSchema_K2::IsGraphMarkedThreadSafe(const UEdGraph* InGraph) const
+{
+	TArray<UK2Node_FunctionEntry*> EntryNodes;
+	InGraph->GetNodesOfClass(EntryNodes);
+
+	for(UK2Node_FunctionEntry* EntryNode : EntryNodes)
+	{
+		if(EntryNode->MetaData.bThreadSafe)
+		{
+			return true;
+		}
+		else if(UFunction* Function = FFunctionFromNodeHelper::FunctionFromNode(EntryNode))
+		{
+			if(FBlueprintEditorUtils::HasFunctionBlueprintThreadSafeMetaData(Function))
+			{
+				return true;
+			}
+		}
+	}
+
+	return false;
+}
+
 bool UEdGraphSchema_K2::IsPropertyExposedOnSpawn(const FProperty* Property)
 {
 	Property = FBlueprintEditorUtils::GetMostUpToDateProperty(Property);
diff --git a/Engine/Source/Editor/BlueprintGraph/Public/BlueprintActionFilter.h b/Engine/Source/Editor/BlueprintGraph/Public/BlueprintActionFilter.h
--- a/Engine/Source/Editor/BlueprintGraph/Public/BlueprintActionFilter.h
+++ b/Engine/Source/Editor/BlueprintGraph/Public/BlueprintActionFilter.h
@@ -194,6 +194,12 @@ public:
 		 * imported namespaces will be filtered out.
 		 */
 		BPFILTER_RejectOutOfScopeMembers	= (1<<4),
+
+		/**
+		 * Rejects actions that would spawn nodes with thread safety characteristics
+		 * that are incompatible with the specification of the containing graph
+		 */
+		BPFILTER_RejectIncompatibleThreadSafety = (1<<5)
 	};
 	FBlueprintActionFilter(uint32 const Flags = 0x00);
 	
diff --git a/Engine/Source/Editor/Kismet/Private/BlueprintActionMenuUtils.cpp b/Engine/Source/Editor/Kismet/Private/BlueprintActionMenuUtils.cpp
--- a/Engine/Source/Editor/Kismet/Private/BlueprintActionMenuUtils.cpp
+++ b/Engine/Source/Editor/Kismet/Private/BlueprintActionMenuUtils.cpp
@@ -461,6 +461,11 @@ void FBlueprintActionMenuUtils::MakeContextMenu(FBlueprintActionContext const& C
 		FilterFlags |= FBlueprintActionFilter::BPFILTER_RejectOutOfScopeMembers;
 	}
 
+	if(bIsContextSensitive)
+	{
+		FilterFlags |= FBlueprintActionFilter::BPFILTER_RejectIncompatibleThreadSafety;
+	}
+
 	FBlueprintActionFilter MainMenuFilter(FilterFlags);
 	MainMenuFilter.Context = Context;
 	MainMenuFilter.Context.SelectedObjects.Empty();
diff --git a/Engine/Source/Editor/Kismet/Private/BlueprintDetailsCustomization.cpp b/Engine/Source/Editor/Kismet/Private/BlueprintDetailsCustomization.cpp
--- a/Engine/Source/Editor/Kismet/Private/BlueprintDetailsCustomization.cpp
+++ b/Engine/Source/Editor/Kismet/Private/BlueprintDetailsCustomization.cpp
@@ -3508,6 +3508,23 @@ void FBlueprintGraphActionDetails::CustomizeDetails( IDetailLayoutBuilder& Detai
 					.OnCheckStateChanged( this, &FBlueprintGraphActionDetails::OnIsExecFunctionModified )
 				];
 			}
+			if (IsThreadSafeFunctionVisible())
+			{
+				Category.AddCustomRow( LOCTEXT( "FunctionThreadSafe_Tooltip", "Thread Safe" ), true )
+				.NameContent()
+				[
+					SNew(STextBlock)
+					.Text( LOCTEXT( "FunctionThreadSafe_Tooltip", "Thread Safe" ) )
+				 	.ToolTipText( LOCTEXT("FunctionIsThreadSafe_Tooltip", "Enable thread-safety checks on this function. Only thread-safe functions and operations are allowed in this function.") )
+					.Font( IDetailLayoutBuilder::GetDetailFont() )
+				]
+				.ValueContent()
+				[
+					SNew( SCheckBox )
+					.IsChecked( this, &FBlueprintGraphActionDetails::GetIsThreadSafeFunction )
+					.OnCheckStateChanged( this, &FBlueprintGraphActionDetails::OnIsThreadSafeFunctionModified )
+				];
+			}
 		}
 
 		if (bIsCustomEvent)
@@ -5091,6 +5108,48 @@ ECheckBoxState FBlueprintGraphActionDetails::GetIsExecFunction() const
 	return (EntryNode->GetFunctionFlags() & FUNC_Exec) ? ECheckBoxState::Checked : ECheckBoxState::Unchecked;
 }
 
+bool FBlueprintGraphActionDetails::IsThreadSafeFunctionVisible() const
+{
+	bool bSupportedType = false;
+	bool bIsEditable = false;
+	
+	UK2Node_EditablePinBase* FunctionEntryNode = FunctionEntryNodePtr.Get();
+	if(FunctionEntryNode)
+	{
+		bSupportedType = FunctionEntryNode->IsA<UK2Node_FunctionEntry>();
+		bIsEditable = FunctionEntryNode->IsEditable();
+	}
+
+	return bIsEditable && bSupportedType;
+}
+
+void FBlueprintGraphActionDetails::OnIsThreadSafeFunctionModified(const ECheckBoxState NewCheckedState)
+{
+	if( FunctionEntryNodePtr.IsValid() )
+	{
+		const bool bThreadSafe = NewCheckedState == ECheckBoxState::Checked;
+		const FText TransactionType = bThreadSafe ? LOCTEXT( "DisableThreadSafe", "Disable Thread Safe" ) : LOCTEXT( "EnableThreadSafe", "Enable Thread Safe" );
+
+		if( UK2Node_FunctionEntry* EntryPoint = Cast<UK2Node_FunctionEntry>(FunctionEntryNodePtr.Get()) )
+		{
+			const FScopedTransaction Transaction( TransactionType );
+			EntryPoint->MetaData.bThreadSafe = bThreadSafe;
+			FBlueprintEditorUtils::MarkBlueprintAsModified( EntryPoint->GetBlueprint() );
+		}
+	}
+}
+
+ECheckBoxState FBlueprintGraphActionDetails::GetIsThreadSafeFunction() const
+{
+	UK2Node_EditablePinBase* FunctionEntryNode = FunctionEntryNodePtr.Get();
+	UK2Node_FunctionEntry* EntryNode = Cast<UK2Node_FunctionEntry>(FunctionEntryNode);
+	if (!EntryNode)
+	{
+		return ECheckBoxState::Undetermined;
+	}
+	return EntryNode->MetaData.bThreadSafe ? ECheckBoxState::Checked : ECheckBoxState::Unchecked;
+}
+
 FReply FBaseBlueprintGraphActionDetails::OnAddNewInputClicked()
 {
 	UK2Node_EditablePinBase* FunctionEntryNode = FunctionEntryNodePtr.Get();
diff --git a/Engine/Source/Editor/Kismet/Private/BlueprintDetailsCustomization.h b/Engine/Source/Editor/Kismet/Private/BlueprintDetailsCustomization.h
--- a/Engine/Source/Editor/Kismet/Private/BlueprintDetailsCustomization.h
+++ b/Engine/Source/Editor/Kismet/Private/BlueprintDetailsCustomization.h
@@ -614,6 +614,10 @@ private:
 	void OnIsExecFunctionModified(const ECheckBoxState NewCheckedState);
 	ECheckBoxState GetIsExecFunction() const;
 
+	bool IsThreadSafeFunctionVisible() const;
+	void OnIsThreadSafeFunctionModified(const ECheckBoxState NewCheckedState);
+	ECheckBoxState GetIsThreadSafeFunction() const;
+	
 	/** Determines if the selected event is identified as editor callable */
 	ECheckBoxState GetIsEditorCallableEvent() const;
 
diff --git a/Engine/Source/Editor/KismetCompiler/Private/KismetCompiler.cpp b/Engine/Source/Editor/KismetCompiler/Private/KismetCompiler.cpp
--- a/Engine/Source/Editor/KismetCompiler/Private/KismetCompiler.cpp
+++ b/Engine/Source/Editor/KismetCompiler/Private/KismetCompiler.cpp
@@ -2229,6 +2229,12 @@ void FKismetCompilerContext::CompileFunction(FKismetFunctionContext& Context)
 		Context.UnsortedSeparateExecutionGroups = FKismetCompilerUtilities::FindUnsortedSeparateExecutionGroups(Context.LinearExecutionList);
 	}
 
+	// Propagate thread-safe flags in this first pass. Also gets called from SetCalculatedMetaDataAndFlags in the second
+	// pass to catch skeleton class generation
+	if (Context.EntryPoint->MetaData.bThreadSafe)
+	{
+		Context.Function->SetMetaData(FBlueprintMetadata::MD_ThreadSafe, TEXT("true"));
+	}
 }
 
 /**
@@ -2276,6 +2282,12 @@ void FKismetCompilerContext::FinishCompilingFunction(FKismetFunctionContext& Con
 		}
 	}
 #endif//VALIDATE_UBER_GRAPH_PERSISTENT_FRAME
+
+	// Check thread safety
+	if (FBlueprintEditorUtils::HasFunctionBlueprintThreadSafeMetaData(Context.Function))
+	{
+		FKismetCompilerUtilities::CheckFunctionThreadSafety(Context, MessageLog);
+	}
 }
 
 void FKismetCompilerContext::SetCalculatedMetaDataAndFlags(UFunction* Function, UK2Node_FunctionEntry* EntryNode, const UEdGraphSchema_K2* K2Schema)
@@ -2360,6 +2372,12 @@ void FKismetCompilerContext::SetCalculatedMetaDataAndFlags(UFunction* Function,
 			Function->SetMetaData(FBlueprintMetadata::MD_DeprecationMessage, *(EntryNode->MetaData.DeprecationMessage));
 		}
 	}
+
+	// Just copy the thread-safe flag with no checking, we verify thread safety elsewhere
+	if (EntryNode->MetaData.bThreadSafe)
+	{
+		Function->SetMetaData(FBlueprintMetadata::MD_ThreadSafe, TEXT("true"));
+	}
 	
 	if (UEdGraphPin* WorldContextPin = EntryNode->GetAutoWorldContextPin())
 	{
@@ -2372,6 +2390,12 @@ void FKismetCompilerContext::SetCalculatedMetaDataAndFlags(UFunction* Function,
 	{
 		// Copy metadata from parent function as well
 		UMetaData::CopyMetadata(OverriddenFunction, Function);
+
+		// Native thread safety metadata works on the class level as well as functions, so we need to check that too.
+		if(FBlueprintEditorUtils::HasFunctionBlueprintThreadSafeMetaData(OverriddenFunction))
+		{
+			Function->SetMetaData(FBlueprintMetadata::MD_ThreadSafe, TEXT("true"));
+		}
 	}
 }
 
diff --git a/Engine/Source/Editor/KismetCompiler/Private/KismetCompilerMisc.cpp b/Engine/Source/Editor/KismetCompiler/Private/KismetCompilerMisc.cpp
--- a/Engine/Source/Editor/KismetCompiler/Private/KismetCompilerMisc.cpp
+++ b/Engine/Source/Editor/KismetCompiler/Private/KismetCompilerMisc.cpp
@@ -1694,6 +1694,209 @@ void FKismetCompilerUtilities::UpdateDependentBlueprints(UBlueprint* ForBP)
 	}
 }
 
+bool FKismetCompilerUtilities::CheckFunctionThreadSafety(const FKismetFunctionContext& InContext, FCompilerResultsLog& InMessageLog, bool InbEmitErrors)
+{
+	bool bIsThreadSafe = true;
+	
+	for(const TPair<UEdGraphNode*, TArray<FBlueprintCompiledStatement*>>& StatementPair : InContext.StatementsPerNode)
+	{
+		bIsThreadSafe &= CheckFunctionCompiledStatementsThreadSafety(StatementPair.Key, InContext.SourceGraph, StatementPair.Value, InMessageLog, InbEmitErrors);
+	}
+
+	return bIsThreadSafe;
+}
+
+// Helper used to emit to log as errors/warnings 
+struct FLogThreadSafetyHelper
+{
+	FLogThreadSafetyHelper(FCompilerResultsLog& InLog, bool bInEmitErrors)
+		: Log(InLog)
+		, bEmitErrors(bInEmitErrors)
+	{}
+		
+	FCompilerResultsLog& Log;
+	bool bEmitErrors;
+		
+	template<typename... Args>
+	void Message(const TCHAR* Format, Args... args)
+	{
+		if(bEmitErrors)
+		{
+			Log.Error(Format, args...);
+		}
+		else
+		{
+			Log.Warning(Format, args...);
+		}
+	}
+};
+
+#define LOG_THREADSAFETY_HELPER(EmitErrors, CategoryName, Format, ...) \
+	if(EmitErrors) \
+	{ \
+		UE_LOG(CategoryName, Error, Format, ##__VA_ARGS__); \
+	} \
+	else \
+	{ \
+		UE_LOG(CategoryName, Warning, Format, ##__VA_ARGS__); \
+	}
+
+bool FKismetCompilerUtilities::CheckFunctionCompiledStatementsThreadSafety(const UEdGraphNode* InNode, const UEdGraph* InSourceGraph, const TArray<FBlueprintCompiledStatement*>& InStatements, FCompilerResultsLog& InMessageLog, bool InbEmitErrors)
+{
+	bool bIsThreadSafe = true;
+
+	const FText GenericThreadSafetyErrorOneParam = LOCTEXT("ThreadSafety_Error_Generic", "This is not thread safe when compiled. See the output log for more details.");
+
+	FLogThreadSafetyHelper LogHelper(InMessageLog, InbEmitErrors);
+	
+	for(const FBlueprintCompiledStatement* Statement : InStatements)
+	{
+		auto LogDelegateUsage = [&bIsThreadSafe, &LogHelper, InNode, InbEmitErrors]()
+		{
+			LogHelper.Message(*LOCTEXT("ThreadSafety_Error_Delegate", "@@ Delegate usage is not thread-safe").ToString(), InNode);
+			bIsThreadSafe = false;
+		};
+
+		auto CheckForInvalidInstancedObjectContext = [&bIsThreadSafe, &LogHelper, InNode, &GenericThreadSafetyErrorOneParam, InbEmitErrors](const FBPTerminal* InTerm)
+		{
+			const FBPTerminal* Context = InTerm;
+			while(Context)
+			{
+				if(Context != nullptr && Context->IsObjectContextType() && Context->IsInstancedVarTerm())
+				{
+					if(Context->SourcePin && Context->SourcePin->GetOwningNode())
+					{
+						// @TODO: we could possibly make exceptions for 'assets' here
+						LogHelper.Message(*LOCTEXT("ThreadSafety_Error_InstancedObjectWithPin", "@@ Accessing an object reference is not thread-safe").ToString(), Context->SourcePin->GetOwningNode());
+					}
+					else
+					{
+						LogHelper.Message(*GenericThreadSafetyErrorOneParam.ToString(), InNode);
+						LOG_THREADSAFETY_HELPER(InbEmitErrors, LogBlueprint, TEXT("Expression that accesses an instanced object context is not thread-safe"));
+					}
+					bIsThreadSafe = false;
+				}
+
+				Context = Context->Context;
+			}
+		};
+
+		switch (Statement->Type)
+		{
+			case KCST_Nop:
+				break;
+			case KCST_CallFunction:
+			{
+				check(Statement->FunctionToCall);
+
+				if(Statement->FunctionContext)
+				{
+					CheckForInvalidInstancedObjectContext(Statement->FunctionContext);
+				}
+				
+				// Check RHS (function inputs) for invalid object access 
+				for(FBPTerminal* RHSTerm : Statement->RHS)
+				{
+					if(RHSTerm->Context)
+					{
+						CheckForInvalidInstancedObjectContext(RHSTerm->Context);
+					}
+				}
+
+				if(!FBlueprintEditorUtils::HasFunctionBlueprintThreadSafeMetaData(Statement->FunctionToCall))
+				{
+					LogHelper.Message(*LOCTEXT("ThreadSafety_Error_NativeFunction", "@@ Non-thread safe function @@ called from thread-safe graph @@").ToString(), InNode, Statement->FunctionToCall, InSourceGraph);
+					bIsThreadSafe = false;
+				}
+				break;
+			}
+			case KCST_Assignment:
+			{
+				// Check LHS/RHS for invalid object access.
+				if(Statement->LHS && Statement->LHS->Context)
+				{
+					CheckForInvalidInstancedObjectContext(Statement->LHS->Context);
+				}
+					
+				for(FBPTerminal* RHSTerm : Statement->RHS)
+				{
+					if(RHSTerm->Context)
+					{
+						CheckForInvalidInstancedObjectContext(RHSTerm->Context);
+					}
+				}
+				break;
+			}
+			case KCST_CompileError:
+		    case KCST_UnconditionalGoto:
+			case KCST_PushState:
+			case KCST_GotoIfNot:
+			case KCST_Return:
+			case KCST_EndOfThread:
+			case KCST_Comment:
+			case KCST_ComputedGoto:
+			case KCST_EndOfThreadIfNot:
+			case KCST_DebugSite:
+			case KCST_CastObjToInterface:
+			case KCST_DynamicCast:
+			case KCST_ObjectToBool:
+				break;
+			case KCST_AddMulticastDelegate:
+			case KCST_ClearMulticastDelegate:
+				LogDelegateUsage();
+				break;
+			case KCST_WireTraceSite:
+				break;
+			case KCST_BindDelegate:
+			case KCST_RemoveMulticastDelegate:
+			case KCST_CallDelegate:
+				LogDelegateUsage();
+				break;
+			case KCST_CreateArray:
+			case KCST_CrossInterfaceCast:
+			case KCST_MetaCast:
+				break;
+			case KCST_AssignmentOnPersistentFrame:
+				LogHelper.Message(*GenericThreadSafetyErrorOneParam.ToString(), InSourceGraph);
+				LOG_THREADSAFETY_HELPER(InbEmitErrors, LogBlueprint, TEXT("Persistent frame assignment is not supported in thread-safe function"));
+				bIsThreadSafe = false;
+				break;
+			case KCST_CastInterfaceToObj:
+			case KCST_GotoReturn:
+			case KCST_GotoReturnIfNot:
+			case KCST_SwitchValue:
+				break;
+			case KCST_InstrumentedEvent:
+			case KCST_InstrumentedEventStop:
+			case KCST_InstrumentedPureNodeEntry:
+			case KCST_InstrumentedWireEntry:
+			case KCST_InstrumentedWireExit:
+			case KCST_InstrumentedStatePush:
+			case KCST_InstrumentedStateRestore:
+			case KCST_InstrumentedStateReset:
+			case KCST_InstrumentedStateSuspend:
+			case KCST_InstrumentedStatePop:
+			case KCST_InstrumentedTunnelEndOfThread:
+				// No way to test instrumentation right now, so this can potentially be removed
+				LogHelper.Message(*GenericThreadSafetyErrorOneParam.ToString(), InSourceGraph);
+				LOG_THREADSAFETY_HELPER(InbEmitErrors, LogBlueprint, TEXT("Instrumentation not supported in thread-safe function"));
+				bIsThreadSafe = false;
+				break;
+			case KCST_ArrayGetByRef:
+			case KCST_CreateSet:
+			case KCST_CreateMap:
+				break;
+			default:
+				LogHelper.Message(*GenericThreadSafetyErrorOneParam.ToString(), InSourceGraph);
+				LOG_THREADSAFETY_HELPER(InbEmitErrors, LogBlueprint, TEXT("Non-thread safe unknown statement type %d (from %s) used in thread-safe context %s"), (int32)Statement->Type, *InNode->GetName(), *InSourceGraph->GetName());
+				bIsThreadSafe = false;
+				break;
+		}
+	}
+
+	return bIsThreadSafe;
+}
+
 //////////////////////////////////////////////////////////////////////////
 // FNodeHandlingFunctor
 
diff --git a/Engine/Source/Editor/KismetCompiler/Public/KismetCompilerMisc.h b/Engine/Source/Editor/KismetCompiler/Public/KismetCompilerMisc.h
--- a/Engine/Source/Editor/KismetCompiler/Public/KismetCompilerMisc.h
+++ b/Engine/Source/Editor/KismetCompiler/Public/KismetCompilerMisc.h
@@ -111,6 +111,12 @@ public:
 
 	/** Add this BP to any BPs that it in*/
 	static void UpdateDependentBlueprints(UBlueprint* BP);
+
+	// Check the passed-in function to verify its thread safety. This makes sure that it only uses/calls thread-safe functions/nodes.
+	static bool CheckFunctionThreadSafety(const FKismetFunctionContext& InContext, FCompilerResultsLog& InMessageLog, bool InbEmitErrors = true);
+
+	// Helper function used by CheckFunctionThreadSafety. Split out to allow the ability to examine individual compiled statement lists (e.g. for the ubergraph)
+	static bool CheckFunctionCompiledStatementsThreadSafety(const UEdGraphNode* InNode, const UEdGraph* InSourceGraph, const TArray<FBlueprintCompiledStatement*>& InStatements, FCompilerResultsLog& InMessageLog, bool InbEmitErrors = true);
 };
 
 //////////////////////////////////////////////////////////////////////////
diff --git a/Engine/Source/Editor/UnrealEd/Private/Kismet2/BlueprintEditorUtils.cpp b/Engine/Source/Editor/UnrealEd/Private/Kismet2/BlueprintEditorUtils.cpp
--- a/Engine/Source/Editor/UnrealEd/Private/Kismet2/BlueprintEditorUtils.cpp
+++ b/Engine/Source/Editor/UnrealEd/Private/Kismet2/BlueprintEditorUtils.cpp
@@ -2410,7 +2410,7 @@ void FBlueprintEditorUtils::CreateMatchingFunction(UK2Node_CallFunction* InNode,
 
 bool FBlueprintEditorUtils::IsFunctionConvertableToEvent(UBlueprint* const BlueprintObj, UFunction* const Function)
 {
-	return BlueprintObj && BlueprintObj->BlueprintType != BPTYPE_FunctionLibrary && BlueprintObj->BlueprintType != BPTYPE_Interface && Function;
+	return BlueprintObj && BlueprintObj->BlueprintType != BPTYPE_FunctionLibrary && BlueprintObj->BlueprintType != BPTYPE_Interface && Function && !HasFunctionBlueprintThreadSafeMetaData(Function);
 }
 
 UFunction* FBlueprintEditorUtils::FindFunctionInImplementedInterfaces(const UBlueprint* Blueprint, const FName& FunctionName, bool * bOutInvalidInterface, bool bGetAllInterfaces)
@@ -10074,4 +10074,24 @@ void FBlueprintEditorUtils::RecombineNestedSubPins(UK2Node* Node)
 	}
 }
 
+bool FBlueprintEditorUtils::HasFunctionBlueprintThreadSafeMetaData(const UFunction* InFunction)
+{
+	if(InFunction)
+	{
+		const bool bHasThreadSafeMetaData = InFunction->HasMetaData(FBlueprintMetadata::MD_ThreadSafe);
+		const bool bHasNotThreadSafeMetaData = InFunction->HasMetaData(FBlueprintMetadata::MD_NotThreadSafe);
+		const bool bClassHasThreadSafeMetaData = InFunction->GetOwnerClass() && InFunction->GetOwnerClass()->HasMetaData(FBlueprintMetadata::MD_ThreadSafe);
+
+		// Native (or BP event) functions need to just have the correct class/function metadata
+		const bool bThreadSafeNative = InFunction->HasAnyFunctionFlags(FUNC_Native | FUNC_BlueprintEvent) && (bHasThreadSafeMetaData || (bClassHasThreadSafeMetaData && !bHasNotThreadSafeMetaData));
+
+		// Script functions get their flag propagated from their entry point, and dont pay heed to class metadata
+		const bool bThreadSafeScript = !InFunction->HasAnyFunctionFlags(FUNC_Native) && bHasThreadSafeMetaData;
+		
+		return bThreadSafeNative || bThreadSafeScript;
+	}
+	
+	return false;
+};
+
 #undef LOCTEXT_NAMESPACE
diff --git a/Engine/Source/Editor/UnrealEd/Public/Kismet2/BlueprintEditorUtils.h b/Engine/Source/Editor/UnrealEd/Public/Kismet2/BlueprintEditorUtils.h
--- a/Engine/Source/Editor/UnrealEd/Public/Kismet2/BlueprintEditorUtils.h
+++ b/Engine/Source/Editor/UnrealEd/Public/Kismet2/BlueprintEditorUtils.h
@@ -1885,6 +1885,14 @@ public:
 	 * Returns true if this BP is currently based on a type that returns true for the UObject::ImplementsGetWorld() call:
 	 */
 	static bool ImplementsGetWorld(const UBlueprint* BP);
+
+	/*
+	 * Checks a function for 'blueprint thread safety'. Note this is not strict thread safety and does not perform any
+	 * analysis on the function, it only inspects the metadata. Analysis is performed by the compiler for BP functions
+	 * that are marked thread safe and metadata is used both here and in the compiler for native functions that have
+	 * been validated for thread safe use in a blueprint.
+	 */
+	static bool HasFunctionBlueprintThreadSafeMetaData(const UFunction* InFunction);
 };
 
 struct UNREALED_API FBlueprintDuplicationScopeFlags
diff --git a/Engine/Source/Runtime/Engine/Classes/Kismet/BlueprintMapLibrary.h b/Engine/Source/Runtime/Engine/Classes/Kismet/BlueprintMapLibrary.h
--- a/Engine/Source/Runtime/Engine/Classes/Kismet/BlueprintMapLibrary.h
+++ b/Engine/Source/Runtime/Engine/Classes/Kismet/BlueprintMapLibrary.h
@@ -10,7 +10,7 @@
 #include "Kismet/BlueprintFunctionLibrary.h"
 #include "BlueprintMapLibrary.generated.h"
 
-UCLASS()
+UCLASS(meta=(BlueprintThreadSafe))
 class ENGINE_API UBlueprintMapLibrary : public UBlueprintFunctionLibrary
 {
 	GENERATED_BODY()
diff --git a/Engine/Source/Runtime/Engine/Classes/Kismet/BlueprintSetLibrary.h b/Engine/Source/Runtime/Engine/Classes/Kismet/BlueprintSetLibrary.h
--- a/Engine/Source/Runtime/Engine/Classes/Kismet/BlueprintSetLibrary.h
+++ b/Engine/Source/Runtime/Engine/Classes/Kismet/BlueprintSetLibrary.h
@@ -10,7 +10,7 @@
 #include "Kismet/BlueprintFunctionLibrary.h"
 #include "BlueprintSetLibrary.generated.h"
 
-UCLASS()
+UCLASS(meta=(BlueprintThreadSafe))
 class ENGINE_API UBlueprintSetLibrary : public UBlueprintFunctionLibrary
 {
 	GENERATED_BODY()
diff --git a/Engine/Source/Runtime/Engine/Classes/Kismet/KismetArrayLibrary.h b/Engine/Source/Runtime/Engine/Classes/Kismet/KismetArrayLibrary.h
--- a/Engine/Source/Runtime/Engine/Classes/Kismet/KismetArrayLibrary.h
+++ b/Engine/Source/Runtime/Engine/Classes/Kismet/KismetArrayLibrary.h
@@ -15,7 +15,7 @@
 
 class AActor;
 
-UCLASS()
+UCLASS(meta=(BlueprintThreadSafe))
 class ENGINE_API UKismetArrayLibrary : public UBlueprintFunctionLibrary
 {
 	GENERATED_UCLASS_BODY()
diff --git a/Engine/Source/Runtime/Engine/Classes/Kismet/KismetNodeHelperLibrary.h b/Engine/Source/Runtime/Engine/Classes/Kismet/KismetNodeHelperLibrary.h
--- a/Engine/Source/Runtime/Engine/Classes/Kismet/KismetNodeHelperLibrary.h
+++ b/Engine/Source/Runtime/Engine/Classes/Kismet/KismetNodeHelperLibrary.h
@@ -12,7 +12,7 @@
 #include "Kismet/BlueprintFunctionLibrary.h"
 #include "KismetNodeHelperLibrary.generated.h"
 
-UCLASS()
+UCLASS(meta=(BlueprintThreadSafe))
 class ENGINE_API UKismetNodeHelperLibrary : public UBlueprintFunctionLibrary
 {
 	GENERATED_UCLASS_BODY()
@@ -84,7 +84,7 @@ class ENGINE_API UKismetNodeHelperLibrary : public UBlueprintFunctionLibrary
 	 * @param bRandom - Whether to select a random index or not
 	 * @return - The index that was selected (returns INDEX_NONE if there was no unmarked bits to choose from)
 	 */
-	UFUNCTION(BlueprintCallable, meta=(BlueprintInternalUseOnly = "TRUE"))
+	UFUNCTION(BlueprintCallable, meta=(BlueprintInternalUseOnly = "TRUE", NotBlueprintThreadSafe))
 	static int32 GetUnmarkedBit(int32 Data, int32 StartIdx, int32 NumBits, bool bRandom);
 
 	/**
@@ -94,7 +94,7 @@ class ENGINE_API UKismetNodeHelperLibrary : public UBlueprintFunctionLibrary
 	 * @param NumBits - The logical number of bits we want to track
 	 * @return - The index that was selected (returns INDEX_NONE if there was no unmarked bits to choose from)
 	 */
-	UFUNCTION(BlueprintCallable, meta=(BlueprintInternalUseOnly = "TRUE"))
+	UFUNCTION(BlueprintCallable, meta=(BlueprintInternalUseOnly = "TRUE", NotBlueprintThreadSafe))
 	static int32 GetRandomUnmarkedBit(int32 Data, int32 StartIdx, int32 NumBits);
 
 	/**
