Improved and added optional limits around private members to Blueprint thread safety checks

Currently it is not 'safe' for external code/scripts to be able to 'poke' data that is ostensibly thread-safe. Restricting to only private data is one possible option to solve this.
Gated behind a config flag for now as most content does not obey these rules at the moment.

This change also covers a case with 'valid' UObject use on worker threads where those UObjects are returned from thread-safe functions (which can only be native given the BP-level checks). It is assumed that the native thread-safe fucntion returning a UObject 'knows' that the UObject should be safe to access (e.g. some sub-object like a linked anim instance).

diff --git a/Engine/Config/BaseEngine.ini b/Engine/Config/BaseEngine.ini
--- a/Engine/Config/BaseEngine.ini
+++ b/Engine/Config/BaseEngine.ini
@@ -1584,6 +1584,7 @@ bLoadNativeConvertedBPClassInEditor=false
 bIgnoreCompileOnLoadErrorsOnBuildMachine=false
 bAllowNativeComponentClassOverrides=true
 bUseLegacyAnimBlueprintThreadSafetyChecks=true
+bThreadSafetyStrictPrivateMemberChecks=false
 
 [/Script/Engine.Blueprint]
 bRecompileOnLoad=true
diff --git a/Engine/Source/Editor/KismetCompiler/Private/KismetCompilerMisc.cpp b/Engine/Source/Editor/KismetCompiler/Private/KismetCompilerMisc.cpp
--- a/Engine/Source/Editor/KismetCompiler/Private/KismetCompilerMisc.cpp
+++ b/Engine/Source/Editor/KismetCompiler/Private/KismetCompilerMisc.cpp
@@ -1762,7 +1762,7 @@ bool FKismetCompilerUtilities::CheckFunctionCompiledStatementsThreadSafety(const
 			const FBPTerminal* Context = InTerm;
 			while(Context)
 			{
-				if(Context != nullptr && Context->IsObjectContextType() && Context->IsInstancedVarTerm())
+				if(Context != nullptr && Context->IsObjectContextType() && Context->Type.PinSubCategoryObject.IsValid() && (Context->IsInstancedVarTerm() || Context->IsLocalVarTerm()))
 				{
 					if(Context->SourcePin && Context->SourcePin->GetOwningNode())
 					{
@@ -1781,6 +1781,40 @@ bool FKismetCompilerUtilities::CheckFunctionCompiledStatementsThreadSafety(const
 			}
 		};
 
+		auto CheckForPrivateMemberUsage = [&bIsThreadSafe, &LogHelper, InNode, &GenericThreadSafetyErrorOneParam, InbEmitErrors](FBPTerminal* InTerm)
+		{
+			static const FBoolConfigValueHelper ThreadSafetyStrictPrivateMemberChecks(TEXT("Kismet"), TEXT("bThreadSafetyStrictPrivateMemberChecks"), GEngineIni);
+			if (ThreadSafetyStrictPrivateMemberChecks)
+			{
+				const FBPTerminal* Context = InTerm;
+				while(Context)
+				{
+					// Check for assignment only to private object variables
+					if(Context->AssociatedVarProperty && !FBlueprintEditorUtils::IsPropertyPrivate(InTerm->AssociatedVarProperty))
+					{
+						if(Context->Context == nullptr && Context->IsInstancedVarTerm() && Context->IsObjectContextType())
+						{
+							UEdGraphNode* OwningNode = Context->SourcePin ? Context->SourcePin->GetOwningNode() : nullptr;
+							if(OwningNode)
+							{
+								LogHelper.Message(*LOCTEXT("ThreadSafety_Error_NonPrivateMemberAccess", "@@ Accessing non-private member variables is not thread-safe. Make the variable private or use a local variable.").ToString(), OwningNode);
+								UE_LOG(LogBlueprint, Display, TEXT("Expression that accesses non-private property '%s' is not thread-safe. This message can be disabled using bThreadSafetyStrictPrivateMemberChecks in Engine.ini"), *Context->AssociatedVarProperty->GetName())
+							}
+							else 
+							{
+								LogHelper.Message(*GenericThreadSafetyErrorOneParam.ToString(), InNode);
+								LOG_THREADSAFETY_HELPER(InbEmitErrors, LogBlueprint, TEXT("Expression that accesses non-private property '%s' is not thread-safe. This message can be disabled using bThreadSafetyStrictPrivateMemberChecks in Engine.ini"), *Context->AssociatedVarProperty->GetName());
+							}
+						}
+
+						bIsThreadSafe = false;
+					}
+
+					Context = Context->Context;
+				}
+			}
+		};
+		
 		switch (Statement->Type)
 		{
 			case KCST_Nop:
@@ -1793,7 +1827,7 @@ bool FKismetCompilerUtilities::CheckFunctionCompiledStatementsThreadSafety(const
 				{
 					CheckForInvalidInstancedObjectContext(Statement->FunctionContext);
 				}
-				
+
 				// Check RHS (function inputs) for invalid object access 
 				for(FBPTerminal* RHSTerm : Statement->RHS)
 				{
@@ -1801,11 +1835,24 @@ bool FKismetCompilerUtilities::CheckFunctionCompiledStatementsThreadSafety(const
 					{
 						CheckForInvalidInstancedObjectContext(RHSTerm->Context);
 					}
+
+					CheckForPrivateMemberUsage(RHSTerm);
 				}
 
 				if(!FBlueprintEditorUtils::HasFunctionBlueprintThreadSafeMetaData(Statement->FunctionToCall))
 				{
-					LogHelper.Message(*LOCTEXT("ThreadSafety_Error_NativeFunction", "@@ Non-thread safe function @@ called from thread-safe graph @@").ToString(), InNode, Statement->FunctionToCall, InSourceGraph);
+					// Check LHS (function return value) for invalid object access.
+					// Note we only do this for BP functions and those that are not declared thread safe. This is to
+					// allow already-useful cases where native code is able to make assumptions about multi-threaded
+					// object access.
+					// A good example of this is returning a 'hosting' anim instance in the context of a linked anim
+					// instance.
+					if(Statement->LHS)
+					{
+						CheckForInvalidInstancedObjectContext(Statement->LHS);
+					}
+					
+					LogHelper.Message(*LOCTEXT("ThreadSafety_Error_NonThreadSafeFunction", "@@ Non-thread safe function @@ called from thread-safe graph @@").ToString(), InNode, Statement->FunctionToCall, InSourceGraph);
 					bIsThreadSafe = false;
 				}
 				break;
@@ -1813,9 +1860,14 @@ bool FKismetCompilerUtilities::CheckFunctionCompiledStatementsThreadSafety(const
 			case KCST_Assignment:
 			{
 				// Check LHS/RHS for invalid object access.
-				if(Statement->LHS && Statement->LHS->Context)
+				if(Statement->LHS)
 				{
-					CheckForInvalidInstancedObjectContext(Statement->LHS->Context);
+					if(Statement->LHS->Context)
+					{
+						CheckForInvalidInstancedObjectContext(Statement->LHS->Context);
+					}
+					
+					CheckForPrivateMemberUsage(Statement->LHS);
 				}
 					
 				for(FBPTerminal* RHSTerm : Statement->RHS)
@@ -1824,6 +1876,8 @@ bool FKismetCompilerUtilities::CheckFunctionCompiledStatementsThreadSafety(const
 					{
 						CheckForInvalidInstancedObjectContext(RHSTerm->Context);
 					}
+
+					CheckForPrivateMemberUsage(RHSTerm);
 				}
 				break;
 			}
diff --git a/Engine/Source/Editor/UnrealEd/Private/Kismet2/BlueprintEditorUtils.cpp b/Engine/Source/Editor/UnrealEd/Private/Kismet2/BlueprintEditorUtils.cpp
--- a/Engine/Source/Editor/UnrealEd/Private/Kismet2/BlueprintEditorUtils.cpp
+++ b/Engine/Source/Editor/UnrealEd/Private/Kismet2/BlueprintEditorUtils.cpp
@@ -8609,6 +8609,11 @@ void FBlueprintEditorUtils::PostEditChangeBlueprintActors(UBlueprint* Blueprint,
 	}
 }
 
+bool FBlueprintEditorUtils::IsPropertyPrivate(const FProperty* Property)
+{
+	return Property->HasAnyPropertyFlags(CPF_NativeAccessSpecifierPrivate) || Property->GetBoolMetaData(FBlueprintMetadata::MD_Private); 
+}
+
 FBlueprintEditorUtils::EPropertyWritableState FBlueprintEditorUtils::IsPropertyWritableInBlueprint(const UBlueprint* Blueprint, const FProperty* Property)
 {
 	if (Property)
diff --git a/Engine/Source/Editor/UnrealEd/Public/Kismet2/BlueprintEditorUtils.h b/Engine/Source/Editor/UnrealEd/Public/Kismet2/BlueprintEditorUtils.h
--- a/Engine/Source/Editor/UnrealEd/Public/Kismet2/BlueprintEditorUtils.h
+++ b/Engine/Source/Editor/UnrealEd/Public/Kismet2/BlueprintEditorUtils.h
@@ -1342,6 +1342,9 @@ public:
 	/** Call PostEditChange() on all Actors based on the given Blueprint */
 	static void PostEditChangeBlueprintActors(UBlueprint* Blueprint, bool bComponentEditChange = false);
 
+	/** @return whether a property is private or not (whether native or via BP metadata) */
+	static bool IsPropertyPrivate(const FProperty* Property);
+	
 	/** Enumeration of whether a property is writable or if not, why. */
 	enum class EPropertyWritableState : uint8
 	{
