Improved Blueprint thread safety checks

Propagate 'thread safe' object terms (thread safe function outputs and parameters) via dynamic casts

diff --git a/Engine/Source/Editor/KismetCompiler/Private/KismetCompilerMisc.cpp b/Engine/Source/Editor/KismetCompiler/Private/KismetCompilerMisc.cpp
--- a/Engine/Source/Editor/KismetCompiler/Private/KismetCompilerMisc.cpp
+++ b/Engine/Source/Editor/KismetCompiler/Private/KismetCompilerMisc.cpp
@@ -1697,10 +1697,84 @@ void FKismetCompilerUtilities::UpdateDependentBlueprints(UBlueprint* ForBP)
 bool FKismetCompilerUtilities::CheckFunctionThreadSafety(const FKismetFunctionContext& InContext, FCompilerResultsLog& InMessageLog, bool InbEmitErrors)
 {
 	bool bIsThreadSafe = true;
+
+	// 1st pass: Build set of 'thread safe' object terms
+	TSet<const FBPTerminal*> ThreadSafeObjectTerms;
 	
+	// Input params to functions (is is assumed that this function is marked thread-safe)
+	if(FBlueprintEditorUtils::HasFunctionBlueprintThreadSafeMetaData(InContext.Function))
+	{
+		for(const FBPTerminal& Parameter : InContext.Parameters)
+		{
+			if(Parameter.IsObjectContextType() && Parameter.IsLocalVarTerm() && Parameter.AssociatedVarProperty && Parameter.AssociatedVarProperty->IsA<FObjectProperty>())
+			{
+				ThreadSafeObjectTerms.Add(&Parameter);
+			}
+		}
+	}
+
 	for(const TPair<UEdGraphNode*, TArray<FBlueprintCompiledStatement*>>& StatementPair : InContext.StatementsPerNode)
 	{
-		bIsThreadSafe &= CheckFunctionCompiledStatementsThreadSafety(StatementPair.Key, InContext.SourceGraph, StatementPair.Value, InMessageLog, InbEmitErrors);
+		for(const FBlueprintCompiledStatement* Statement : StatementPair.Value)
+		{
+			// Return values from thread-safe functions
+			if(Statement->Type == KCST_CallFunction && Statement->FunctionToCall != nullptr)
+			{
+				if(FBlueprintEditorUtils::HasFunctionBlueprintThreadSafeMetaData(Statement->FunctionToCall))
+				{
+					if(Statement->LHS)
+					{
+						ThreadSafeObjectTerms.Add(Statement->LHS);
+					}
+				}
+			}
+
+
+		}
+	}
+
+	// 2nd pass, multiple times: Propagate thread safe terms down the statement lists via supported links
+	// @TODO: we can probably reduce the order of this algorithm by keeping a working set of unchecked terms and only checking them each loop
+	bool bPropagated = false;
+	do
+	{
+		bPropagated = false;
+		
+		for(const TPair<UEdGraphNode*, TArray<FBlueprintCompiledStatement*>>& StatementPair : InContext.StatementsPerNode)
+		{
+			for(const FBlueprintCompiledStatement* Statement : StatementPair.Value)
+			{
+				switch(Statement->Type)
+				{
+				case KCST_CastObjToInterface:
+				case KCST_DynamicCast:
+				case KCST_MetaCast:
+				case KCST_CastInterfaceToObj:
+					if(Statement->LHS)
+					{
+						for(const FBPTerminal* RHSTerm : Statement->RHS)
+						{
+							if(ThreadSafeObjectTerms.Contains(RHSTerm) && !ThreadSafeObjectTerms.Contains(Statement->LHS))
+							{
+								check(Statement->LHS->AssociatedVarProperty && (Statement->LHS->AssociatedVarProperty->IsA<FObjectProperty>() || Statement->LHS->AssociatedVarProperty->IsA<FInterfaceProperty>()));
+								ThreadSafeObjectTerms.Add(Statement->LHS); 
+								bPropagated = true;
+							}
+						}
+					}
+					break;
+				default:
+					break;
+				}
+			}
+		}	
+	}
+	while (bPropagated);
+
+	// 3rd pass: Check statement lists
+	for(const TPair<UEdGraphNode*, TArray<FBlueprintCompiledStatement*>>& StatementPair : InContext.StatementsPerNode)
+	{
+		bIsThreadSafe &= CheckFunctionCompiledStatementsThreadSafety(StatementPair.Key, InContext.SourceGraph, StatementPair.Value, InMessageLog, InbEmitErrors, &ThreadSafeObjectTerms);
 	}
 
 	return bIsThreadSafe;
@@ -1741,7 +1815,7 @@ struct FLogThreadSafetyHelper
 		UE_LOG(CategoryName, Warning, Format, ##__VA_ARGS__); \
 	}
 
-bool FKismetCompilerUtilities::CheckFunctionCompiledStatementsThreadSafety(const UEdGraphNode* InNode, const UEdGraph* InSourceGraph, const TArray<FBlueprintCompiledStatement*>& InStatements, FCompilerResultsLog& InMessageLog, bool InbEmitErrors)
+bool FKismetCompilerUtilities::CheckFunctionCompiledStatementsThreadSafety(const UEdGraphNode* InNode, const UEdGraph* InSourceGraph, const TArray<FBlueprintCompiledStatement*>& InStatements, FCompilerResultsLog& InMessageLog, bool InbEmitErrors, TSet<const FBPTerminal*>* InThreadSafeObjectTerms)
 {
 	bool bIsThreadSafe = true;
 
@@ -1757,27 +1831,33 @@ bool FKismetCompilerUtilities::CheckFunctionCompiledStatementsThreadSafety(const
 			bIsThreadSafe = false;
 		};
 
-		auto CheckForInvalidInstancedObjectContext = [&bIsThreadSafe, &LogHelper, InNode, &GenericThreadSafetyErrorOneParam, InbEmitErrors](const FBPTerminal* InTerm)
+		auto CheckForInvalidInstancedObjectContext = [&bIsThreadSafe, &LogHelper, InNode, &GenericThreadSafetyErrorOneParam, InbEmitErrors, InThreadSafeObjectTerms](const FBPTerminal* InTerm)
 		{
 			const FBPTerminal* Context = InTerm;
 			while(Context)
 			{
-				if(Context != nullptr && Context->IsObjectContextType() && Context->Type.PinSubCategoryObject.IsValid() && (Context->IsInstancedVarTerm() || Context->IsLocalVarTerm()))
+				if(Context != nullptr)
 				{
-					if(Context->SourcePin && Context->SourcePin->GetOwningNode())
-					{
-						// @TODO: we could possibly make exceptions for 'assets' here
-						LogHelper.Message(*LOCTEXT("ThreadSafety_Error_InstancedObjectWithPin", "@@ Accessing an object reference is not thread-safe").ToString(), Context->SourcePin->GetOwningNode());
-					}
-					else
+					if(InThreadSafeObjectTerms == nullptr || !InThreadSafeObjectTerms->Contains(Context))
 					{
-						LogHelper.Message(*GenericThreadSafetyErrorOneParam.ToString(), InNode);
-						LOG_THREADSAFETY_HELPER(InbEmitErrors, LogBlueprint, TEXT("Expression that accesses an instanced object context is not thread-safe"));
+						if(Context->IsObjectContextType() && Context->Type.PinSubCategoryObject.IsValid() && (Context->IsInstancedVarTerm() || Context->IsLocalVarTerm()))
+						{
+							if(Context->SourcePin && Context->SourcePin->GetOwningNode())
+							{
+								// @TODO: we could possibly make exceptions for 'assets' here
+								LogHelper.Message(*LOCTEXT("ThreadSafety_Error_InstancedObjectWithPin", "@@ Accessing an object reference is not thread-safe").ToString(), Context->SourcePin->GetOwningNode());
+							}
+							else
+							{
+								LogHelper.Message(*GenericThreadSafetyErrorOneParam.ToString(), InNode);
+								LOG_THREADSAFETY_HELPER(InbEmitErrors, LogBlueprint, TEXT("Expression that accesses an instanced object context is not thread-safe"));
+							}
+							bIsThreadSafe = false;
+						}
 					}
-					bIsThreadSafe = false;
-				}
 
-				Context = Context->Context;
+					Context = Context->Context;
+				}
 			}
 		};
 
diff --git a/Engine/Source/Editor/KismetCompiler/Public/KismetCompilerMisc.h b/Engine/Source/Editor/KismetCompiler/Public/KismetCompilerMisc.h
--- a/Engine/Source/Editor/KismetCompiler/Public/KismetCompilerMisc.h
+++ b/Engine/Source/Editor/KismetCompiler/Public/KismetCompilerMisc.h
@@ -116,7 +116,7 @@ public:
 	static bool CheckFunctionThreadSafety(const FKismetFunctionContext& InContext, FCompilerResultsLog& InMessageLog, bool InbEmitErrors = true);
 
 	// Helper function used by CheckFunctionThreadSafety. Split out to allow the ability to examine individual compiled statement lists (e.g. for the ubergraph)
-	static bool CheckFunctionCompiledStatementsThreadSafety(const UEdGraphNode* InNode, const UEdGraph* InSourceGraph, const TArray<FBlueprintCompiledStatement*>& InStatements, FCompilerResultsLog& InMessageLog, bool InbEmitErrors = true);
+	static bool CheckFunctionCompiledStatementsThreadSafety(const UEdGraphNode* InNode, const UEdGraph* InSourceGraph, const TArray<FBlueprintCompiledStatement*>& InStatements, FCompilerResultsLog& InMessageLog, bool InbEmitErrors = true, TSet<const FBPTerminal*>* InThreadSafeObjectTerms = nullptr);
 };
 
 //////////////////////////////////////////////////////////////////////////
