Remove transient "TRASHSTRUCT_" objects as they are no longer required by user-defined struct precompilation in the editor context.

The transition from UProperty->FProperty eliminated the previous need to relocate child properties to a temp struct object for de-linking/eventual garbage collection. There are no other subobjects outered to a user-defined struct object that would otherwise need to be handled this way.

diff --git a/Engine/Source/Editor/KismetCompiler/Private/UserDefinedStructureCompilerUtils.cpp b/Engine/Source/Editor/KismetCompiler/Private/UserDefinedStructureCompilerUtils.cpp
--- a/Engine/Source/Editor/KismetCompiler/Private/UserDefinedStructureCompilerUtils.cpp
+++ b/Engine/Source/Editor/KismetCompiler/Private/UserDefinedStructureCompilerUtils.cpp
@@ -144,7 +144,7 @@ struct FUserDefinedStructureCompilerInner
 		}
 	}
 
-	static UObject* CleanAndSanitizeStruct(UUserDefinedStruct* StructToClean)
+	static void CleanAndSanitizeStruct(UUserDefinedStruct* StructToClean)
 	{
 		check(StructToClean);
 
@@ -153,23 +153,8 @@ struct FUserDefinedStructureCompilerInner
 			EditorData->CleanDefaultInstance();
 		}
 
-		UUserDefinedStruct* TransientStruct = nullptr;
-
 		if (FStructureEditorUtils::FStructEditorManager::ActiveChange != FStructureEditorUtils::EStructureEditorChangeInfo::DefaultValueChanged)
 		{
-			const FString TransientString = FString::Printf(TEXT("TRASHSTRUCT_%s"), *StructToClean->GetName());
-			const FName TransientName = MakeUniqueObjectName(GetTransientPackage(), UUserDefinedStruct::StaticClass(), FName(*TransientString));
-			TransientStruct = NewObject<UUserDefinedStruct>(GetTransientPackage(), TransientName, RF_Public | RF_Transient);
-			TransientStruct->PrepareCppStructOps();
-
-			TArray<UObject*> SubObjects;
-			GetObjectsWithOuter(StructToClean, SubObjects, true);
-			SubObjects.Remove(StructToClean->EditorData);
-			for (UObject* CurrSubObj : SubObjects)
-			{
-				FLinkerLoad::InvalidateExport(CurrSubObj);
-			}
-
 			StructToClean->SetSuperStruct(nullptr);
 			StructToClean->Children = nullptr;
 			StructToClean->DestroyChildPropertiesAndResetPropertyLinks();
@@ -179,8 +164,6 @@ struct FUserDefinedStructureCompilerInner
 			StructToClean->ErrorMessage.Empty();
 			StructToClean->SetStructTrashed(true);
 		}
-
-		return TransientStruct;
 	}
 
 	static void LogError(UUserDefinedStruct* Struct, FCompilerResultsLog& MessageLog, const FString& ErrorMsg)
