Favor instance data over sparse data when resolving BP pins to properties

diff --git a/Engine/Source/Editor/KismetCompiler/Private/KismetCompilerMisc.cpp b/Engine/Source/Editor/KismetCompiler/Private/KismetCompilerMisc.cpp
--- a/Engine/Source/Editor/KismetCompiler/Private/KismetCompilerMisc.cpp
+++ b/Engine/Source/Editor/KismetCompiler/Private/KismetCompilerMisc.cpp
@@ -527,98 +527,73 @@ void FKismetCompilerUtilities::RemoveObjectRedirectorIfPresent(UObject* Package,
 }
 
 /** Finds a property by name, starting in the specified scope; Validates property type and returns NULL along with emitting an error if there is a mismatch. */
-FProperty* FKismetCompilerUtilities::FindPropertyInScope(UStruct* Scope, UEdGraphPin* Pin, FCompilerResultsLog& MessageLog, const UEdGraphSchema_K2* Schema, UClass* SelfClass, bool& bIsSparseProperty, bool bSuppressMissingMemberErrors)
+FProperty* FKismetCompilerUtilities::FindPropertyInScope(UStruct* Scope, UEdGraphPin* Pin, FCompilerResultsLog& MessageLog, const UEdGraphSchema_K2* Schema, UClass* SelfClass, bool& bIsSparseProperty)
 {
-	bIsSparseProperty = false;
-	UStruct* InitialScope = Scope;
-	while (Scope != nullptr)
-	{
-		// If this is a class, check the sparse data for the property
-		UClass* Class = Cast<UClass>(Scope);
-		if (Class)
-		{
-			UStruct* SparseData = Class->GetSparseClassDataStruct();
-			if (SparseData)
-			{
-				FProperty* Prop = FindPropertyInScope(SparseData, Pin, MessageLog, Schema, SelfClass, bIsSparseProperty, true);
-				if (Prop)
+	if (FProperty* Property = FKismetCompilerUtilities::FindNamedPropertyInScope(Scope, Pin->PinName, bIsSparseProperty, /*bAllowDeprecated*/true))
 	{
-					bIsSparseProperty = true;
-					return Prop;
-				}
-			}
-		}
-
-		for (TFieldIterator<FProperty> It(Scope, EFieldIteratorFlags::IncludeSuper); It; ++It)
+		if (FKismetCompilerUtilities::IsTypeCompatibleWithProperty(Pin, Property, MessageLog, Schema, SelfClass))
 		{
-			FProperty* Property = *It;
-
-			if (Property->GetFName() == Pin->PinName)
-			{
-				if (FKismetCompilerUtilities::IsTypeCompatibleWithProperty(Pin, Property, MessageLog, Schema, SelfClass))
-				{
-					return Property;
-				}
-				else
-				{
-					// Exit now, we found one with the right name but the type mismatched (and there was a type mismatch error)
-					return nullptr;
-				}
-			}
+			return Property;
 		}
-
-		// Functions don't automatically check their class when using a field iterator
-		UFunction* Function = Cast<UFunction>(Scope);
-		Scope = (Function != nullptr) ? Cast<UStruct>(Function->GetOuter()) : nullptr;
 	}
-
-	// Couldn't find the name
-	if (!FKismetCompilerUtilities::IsMissingMemberPotentiallyLoading(Cast<UBlueprint>(SelfClass->ClassGeneratedBy), InitialScope) && !bSuppressMissingMemberErrors)
+	else if (!FKismetCompilerUtilities::IsMissingMemberPotentiallyLoading(Cast<UBlueprint>(SelfClass->ClassGeneratedBy), Scope))
 	{
 		MessageLog.Error(*FText::Format(LOCTEXT("PropertyNotFound_Error", "The property associated with @@ could not be found in '{0}'"), FText::FromString(SelfClass->GetPathName())).ToString(), Pin);
 	}
+
 	return nullptr;
 }
 
 // Finds a property by name, starting in the specified scope, returning NULL if it's not found
-FProperty* FKismetCompilerUtilities::FindNamedPropertyInScope(UStruct* Scope, FName PropertyName, bool& bIsSparseProperty)
+FProperty* FKismetCompilerUtilities::FindNamedPropertyInScope(UStruct* Scope, FName PropertyName, bool& bIsSparseProperty, const bool bAllowDeprecated)
 {
-	bIsSparseProperty = false;
-	while (Scope != NULL)
+	auto FindProperty = [PropertyName, bAllowDeprecated](UStruct* CurrentScope) -> FProperty*
 	{
-		for (TFieldIterator<FProperty> It(Scope, EFieldIteratorFlags::IncludeSuper); It; ++It)
+		for (TFieldIterator<FProperty> It(CurrentScope); It; ++It)
 		{
 			FProperty* Property = *It;
 
-			// If we match by name, and var is not deprecated...
-			if (Property->GetFName() == PropertyName && !Property->HasAllPropertyFlags(CPF_Deprecated))
+			if (Property->GetFName() == PropertyName)
 			{
-				return Property;
+				if (bAllowDeprecated || !Property->HasAllPropertyFlags(CPF_Deprecated))
+				{
+					return Property;
+				}
+				break;
 			}
 		}
 
+		return nullptr;
+	};
+
+	bIsSparseProperty = false;
+	while (Scope)
+	{
+		// Check the given scope first
+		if (FProperty* Property = FindProperty(Scope))
+		{
+			return Property;
+		}
+
 		// If this is a class, check the sparse data for the property
-		UClass* Class = Cast<UClass>(Scope);
-		if (Class)
+		if (UClass* Class = Cast<UClass>(Scope))
 		{
-			UStruct* SparseData = Class->GetSparseClassDataStruct();
-			if (SparseData)
+			if (UStruct* SparseData = Class->GetSparseClassDataStruct())
 			{
-				FProperty* Prop = FindNamedPropertyInScope(SparseData, PropertyName, bIsSparseProperty);
-				if (Prop)
+				if (FProperty* Property = FindProperty(SparseData))
 				{
 					bIsSparseProperty = true;
-					return Prop;
+					return Property;
 				}
 			}
 		}
 
 		// Functions don't automatically check their class when using a field iterator
 		UFunction* Function = Cast<UFunction>(Scope);
-		Scope = (Function != NULL) ? Cast<UStruct>(Function->GetOuter()) : NULL;
+		Scope = Function ? Cast<UStruct>(Function->GetOuter()) : nullptr;
 	}
 
-	return NULL;
+	return nullptr;
 }
 
 void FKismetCompilerUtilities::CompileDefaultProperties(UClass* Class)
diff --git a/Engine/Source/Editor/KismetCompiler/Public/KismetCompilerMisc.h b/Engine/Source/Editor/KismetCompiler/Public/KismetCompilerMisc.h
--- a/Engine/Source/Editor/KismetCompiler/Public/KismetCompilerMisc.h
+++ b/Engine/Source/Editor/KismetCompiler/Public/KismetCompilerMisc.h
@@ -42,10 +42,10 @@ public:
 	static bool IsTypeCompatibleWithProperty(UEdGraphPin* SourcePin, FProperty* Property, FCompilerResultsLog& MessageLog, const UEdGraphSchema_K2* Schema, UClass* SelfClass);
 
 	/** Finds a property by name, starting in the specified scope; Validates property type and returns NULL along with emitting an error if there is a mismatch. */
-	static FProperty* FindPropertyInScope(UStruct* Scope, UEdGraphPin* Pin, FCompilerResultsLog& MessageLog, const UEdGraphSchema_K2* Schema, UClass* SelfClass, bool& bIsSparseProperty, bool bSuppressMissingMemberErrors = false);
+	static FProperty* FindPropertyInScope(UStruct* Scope, UEdGraphPin* Pin, FCompilerResultsLog& MessageLog, const UEdGraphSchema_K2* Schema, UClass* SelfClass, bool& bIsSparseProperty);
 
 	// Finds a property by name, starting in the specified scope, returning NULL if it's not found
-	static FProperty* FindNamedPropertyInScope(UStruct* Scope, FName PropertyName, bool& bIsSparseProperty);
+	static FProperty* FindNamedPropertyInScope(UStruct* Scope, FName PropertyName, bool& bIsSparseProperty, const bool bAllowDeprecated = false);
 
 	/** return function, that overrides BlueprintImplementableEvent with given name in given class (super-classes are not considered) */
 	static const UFunction* FindOverriddenImplementableEvent(const FName& EventName, const UClass* Class);
