Allow script execution to bypass development-only nodes in a Blueprint function graph when a project is configured to compile them out at cook time.

Previously, this only applied to disabled/development-only nodes in an event graph.

diff --git a/Engine/Source/Editor/KismetCompiler/Private/KismetCompiler.cpp b/Engine/Source/Editor/KismetCompiler/Private/KismetCompiler.cpp
--- a/Engine/Source/Editor/KismetCompiler/Private/KismetCompiler.cpp
+++ b/Engine/Source/Editor/KismetCompiler/Private/KismetCompiler.cpp
@@ -3832,9 +3832,17 @@ void FKismetCompilerContext::ProcessOneFunctionGraph(UEdGraph* SourceGraph, bool
 		return;
 	}
 
-	// Clone the source graph so we can modify it as needed; merging in the child graphs
-	UEdGraph* FunctionGraph = FEdGraphUtilities::CloneGraph(SourceGraph, Blueprint, &MessageLog, true); 
-	FEdGraphUtilities::MergeChildrenGraphsIn(FunctionGraph, FunctionGraph, /*bRequireSchemaMatch=*/ true);
+	// Create an (empty) intermediate graph object
+	UEdGraph* FunctionGraph = NewObject<UEdGraph>(Blueprint);
+	FunctionGraph->Schema = UEdGraphSchema_K2::StaticClass();
+	FunctionGraph->SetFlags(RF_Transient);
+
+	// Clone the source graph so we can modify it as needed; merging it along with any child graphs into the intermediate graph object
+	// This matches how event graphs are handled - each source graph is cloned and merged into the ubergraph, excluding any disabled nodes
+	FEdGraphUtilities::CloneAndMergeGraphIn(FunctionGraph, SourceGraph, MessageLog, /*bRequireSchemaMatch=*/ true, /*bIsCompiling=*/ true);
+
+	// Record the intermediate graph object into the backtrace map, etc.
+	MessageLog.NotifyIntermediateObjectCreation(FunctionGraph, SourceGraph);
 
 	ExpansionStep(FunctionGraph, false);
 
