Back out CL 18640377.

#rnx

diff --git a/Engine/Source/Editor/KismetCompiler/Private/KismetCompiler.cpp b/Engine/Source/Editor/KismetCompiler/Private/KismetCompiler.cpp
--- a/Engine/Source/Editor/KismetCompiler/Private/KismetCompiler.cpp
+++ b/Engine/Source/Editor/KismetCompiler/Private/KismetCompiler.cpp
@@ -3832,17 +3832,9 @@ void FKismetCompilerContext::ProcessOneFunctionGraph(UEdGraph* SourceGraph, bool
 		return;
 	}
 
-	// Create an (empty) intermediate graph object
-	UEdGraph* FunctionGraph = NewObject<UEdGraph>(Blueprint);
-	FunctionGraph->Schema = UEdGraphSchema_K2::StaticClass();
-	FunctionGraph->SetFlags(RF_Transient);
-
-	// Clone the source graph so we can modify it as needed; merging it along with any child graphs into the intermediate graph object
-	// This matches how event graphs are handled - each source graph is cloned and merged into the ubergraph, excluding any disabled nodes
-	FEdGraphUtilities::CloneAndMergeGraphIn(FunctionGraph, SourceGraph, MessageLog, /*bRequireSchemaMatch=*/ true, /*bIsCompiling=*/ true);
-
-	// Record the intermediate graph object into the backtrace map, etc.
-	MessageLog.NotifyIntermediateObjectCreation(FunctionGraph, SourceGraph);
+	// Clone the source graph so we can modify it as needed; merging in the child graphs
+	UEdGraph* FunctionGraph = FEdGraphUtilities::CloneGraph(SourceGraph, Blueprint, &MessageLog, true);
+	FEdGraphUtilities::MergeChildrenGraphsIn(FunctionGraph, FunctionGraph, /*bRequireSchemaMatch=*/ true);
 
 	ExpansionStep(FunctionGraph, false);
 
