Clean up invalid/stale dependent references on Blueprint compile/destroy.

No real impact other than to keep cached dependent sets from growing indefinitely even after weakly-referenced assets are destroyed.

diff --git a/Engine/Source/Editor/KismetCompiler/Private/KismetCompilerMisc.cpp b/Engine/Source/Editor/KismetCompiler/Private/KismetCompilerMisc.cpp
--- a/Engine/Source/Editor/KismetCompiler/Private/KismetCompilerMisc.cpp
+++ b/Engine/Source/Editor/KismetCompiler/Private/KismetCompilerMisc.cpp
@@ -1713,6 +1713,21 @@ void FKismetCompilerUtilities::UpdateDependentBlueprints(UBlueprint* ForBP)
 			}
 		}
 	}
+
+	// Clear out any stale references to invalid/deleted objects.
+	TSet<TWeakObjectPtr<UBlueprint>> InvalidReferences;
+	for (const TWeakObjectPtr<UBlueprint>& Reference : ForBP->CachedDependents)
+	{
+		if (!Reference.IsValid() || Reference.IsStale())
+		{
+			InvalidReferences.Add(Reference);
+		}
+	}
+
+	for (const TWeakObjectPtr<UBlueprint>& InvalidReference : InvalidReferences)
+	{
+		ForBP->CachedDependents.Remove(InvalidReference);
+	}
 }
 
 bool FKismetCompilerUtilities::CheckFunctionThreadSafety(const FKismetFunctionContext& InContext, FCompilerResultsLog& InMessageLog, bool InbEmitErrors)
diff --git a/Engine/Source/Runtime/Engine/Private/Blueprint.cpp b/Engine/Source/Runtime/Engine/Private/Blueprint.cpp
--- a/Engine/Source/Runtime/Engine/Private/Blueprint.cpp
+++ b/Engine/Source/Runtime/Engine/Private/Blueprint.cpp
@@ -1388,6 +1388,15 @@ void UBlueprint::BeginDestroy()
 	Super::BeginDestroy();
 
 	FBlueprintEditorUtils::RemoveAllLocalBookmarks(this);
+
+	// For each cached dependency, remove ourselves from its cached dependent set.
+	for (const TWeakObjectPtr<UBlueprint>& DependencyReference : CachedDependencies)
+	{
+		if (UBlueprint* Dependency = DependencyReference.Get())
+		{
+			Dependency->CachedDependents.Remove(MakeWeakObjectPtr(this));
+		}
+	}
 }
 
 #endif // WITH_EDITOR
