Fixed spurious compiler thread-safety errors when batch compiling

Use most up-to-date function's metadata (from skeleton class) rather than the currently-being-compiled reference

diff --git a/Engine/Source/Editor/KismetCompiler/Private/KismetCompilerMisc.cpp b/Engine/Source/Editor/KismetCompiler/Private/KismetCompilerMisc.cpp
--- a/Engine/Source/Editor/KismetCompiler/Private/KismetCompilerMisc.cpp
+++ b/Engine/Source/Editor/KismetCompiler/Private/KismetCompilerMisc.cpp
@@ -1955,7 +1955,8 @@ bool FKismetCompilerUtilities::CheckFunctionCompiledStatementsThreadSafety(const
 					CheckForPrivateMemberUsage(RHSTerm);
 				}
 
-				if(!FBlueprintEditorUtils::HasFunctionBlueprintThreadSafeMetaData(Statement->FunctionToCall))
+				UFunction* SkeletonClassFunction = FBlueprintEditorUtils::GetMostUpToDateFunction(Statement->FunctionToCall);
+				if(SkeletonClassFunction && !FBlueprintEditorUtils::HasFunctionBlueprintThreadSafeMetaData(SkeletonClassFunction))
 				{
 					// Check LHS (function return value) for invalid object access.
 					// Note we only do this for BP functions and those that are not declared thread safe. This is to
