UE-139626 - Fix for implicit cast warning

The implicit cast warnings didn't take into account nodes that aren't compiled, which is effectively commented out Blueprint code. Since cast fixup occurs during compilation, we were reporting false positives.

diff --git a/Engine/Source/Editor/KismetCompiler/Private/KismetCompiler.cpp b/Engine/Source/Editor/KismetCompiler/Private/KismetCompiler.cpp
--- a/Engine/Source/Editor/KismetCompiler/Private/KismetCompiler.cpp
+++ b/Engine/Source/Editor/KismetCompiler/Private/KismetCompiler.cpp
@@ -2474,13 +2474,37 @@ void FKismetCompilerContext::CompileFunction(FKismetFunctionContext& Context)
 	
 	if (Context.ImplicitCastMap.Num() > 0)
 	{
-		UE_LOG(LogK2Compiler, Warning, TEXT("Unhandled implicit casts found during compilation of function '%s'!"),
-			*Context.Function->GetFullName());
+		// Only issue a warning for pins that belong to nodes in the execution path.
+		// We only compile those particular nodes, which is where implicit cast fixup occurs.
+
+		TArray<UEdGraphPin*> UnhandledPins;
 		for (const auto& It : Context.ImplicitCastMap)
 		{
-			UE_LOG(LogK2Compiler, Warning, TEXT("\tPin '%s' was not handled by node '%s'!"),
-				*It.Key->PinName.ToString(),
-				*It.Key->GetOwningNode()->GetName());
+			UEdGraphNode* OwningNode = It.Key->GetOwningNode();
+			if (SortKeyMap.Find(OwningNode))
+			{
+				UnhandledPins.Add(It.Key);
+			}
+			else
+			{
+				UE_LOG(LogK2Compiler, Verbose, TEXT("Skipping implicit cast check for pin '%s' in function '%s'. Owning node '%s' is not in the execution list."),
+					*It.Key->PinName.ToString(),
+					*Context.Function->GetName(),
+					*It.Key->GetOwningNode()->GetName());
+			}
+		}
+
+		if (UnhandledPins.Num() > 0)
+		{
+			UE_LOG(LogK2Compiler, Warning, TEXT("Unhandled implicit casts found during compilation of function '%s'!"),
+				*Context.Function->GetFullName());
+
+			for (UEdGraphPin* Pin : UnhandledPins)
+			{
+				UE_LOG(LogK2Compiler, Warning, TEXT("\tPin '%s' was not handled by node '%s'!"),
+					*Pin->PinName.ToString(),
+					*Pin->GetOwningNode()->GetName());
+			}
 		}
 	}
 
