Fix for an editor crash after duplicating a Blueprint with a missing/unresolved parent class (regression).

diff --git a/Engine/Source/Editor/KismetCompiler/Private/KismetCompiler.cpp b/Engine/Source/Editor/KismetCompiler/Private/KismetCompiler.cpp
--- a/Engine/Source/Editor/KismetCompiler/Private/KismetCompiler.cpp
+++ b/Engine/Source/Editor/KismetCompiler/Private/KismetCompiler.cpp
@@ -5044,6 +5044,12 @@ void FKismetCompilerContext::CompileFunctions(EInternalCompilerFlags InternalFla
 
 void FKismetCompilerContext::PostCDOCompiled(const UObject::FPostCDOCompiledContext& Context)
 {
+	// Exit now if we did not create a new class.
+	if (!NewClass)
+	{
+		return;
+	}
+
 	// Notify the CDO that it has finished compiling
 	NewClass->ClassDefaultObject->PostCDOCompiled(Context);
 	if (!Context.bIsSkeletonOnly)
