UE-145095 - Create Event  Signature Error: The function/event does not match the necessary signature - has the delegate or function/event changed

A couple of issues were occurring with delegate signature fixup:

1. Unlike the UserDefinedPins array, the Pins array may contain subpins from a split struct pin. We need to intentionally skip these pins since they aren't a part of the function signature.
2. The IsValidParameter lambda had incorrect logic for determining an input parameter. The presence of both CPF_OutParm *and* CPF_ConstParm indicate that this parameter is actually an input parameter.

diff --git a/Engine/Source/Editor/KismetCompiler/Private/KismetCompiler.cpp b/Engine/Source/Editor/KismetCompiler/Private/KismetCompiler.cpp
--- a/Engine/Source/Editor/KismetCompiler/Private/KismetCompiler.cpp
+++ b/Engine/Source/Editor/KismetCompiler/Private/KismetCompiler.cpp
@@ -128,7 +128,7 @@ namespace UE::KismetCompiler::Private
 	// When we change pins back to real/float types, we also expect that the series of pin types matches that of the UFunction.
 	// If there's ever a discrepancy, then there's not much we can do other than log a warning.
 	// It typically means that the BP is in a bad state, which prevents us from deducing the corresponding pin.
-	static void RestoreFloatPinsToNode(UK2Node_EditablePinBase* Node, const UFunction* FunctionSignature, bool bOutputParamsOnly = false)
+	static void RestoreFloatPinsToNode(UK2Node_EditablePinBase* Node, const UFunction* FunctionSignature, bool bProcessOutputParamsOnly = false)
 	{
 		check(Node);
 		check(FunctionSignature);
@@ -148,6 +148,12 @@ namespace UE::KismetCompiler::Private
 			{
 				check(Node->Pins[i]);
 
+				// Subpins aren't a part of the function signature, so we need to skip them.
+				if (Node->Pins[i]->ParentPin)
+				{
+					continue;
+				}
+
 				if (Node->Pins[i]->PinName == UserPinName)
 				{
 					PinCursor = i;
@@ -199,23 +205,26 @@ namespace UE::KismetCompiler::Private
 			return false;
 		};
 
-		auto IsValidParameter = [bOutputParamsOnly](FProperty* Property)
+		auto IsValidParameter = [bProcessOutputParamsOnly](FProperty* Property)
 		{
 			check(Property);
 
-			if (bOutputParamsOnly)
+			bool bIsOutputParameter =
+				(Property->HasAnyPropertyFlags(CPF_OutParm) && !Property->HasAnyPropertyFlags(CPF_ConstParm));
+
+			if (bProcessOutputParamsOnly)
 			{
-				return Property->HasAnyPropertyFlags(CPF_OutParm);
+				return bIsOutputParameter;
 			}
 			else
 			{
 				// HACK: We still have to reckon with arrays that are implicitly copy-by-ref.
 				// Ideally, we should be excluding all output params.
-				bool bIsValidInput =
+				bool bIsInputParameter =
 					(Property->IsA<FArrayProperty>() && Property->HasAnyPropertyFlags(CPF_Parm)) ||
-					(Property->HasAnyPropertyFlags(CPF_Parm) && !Property->HasAnyPropertyFlags(CPF_OutParm));
+					!bIsOutputParameter;
 
-				return bIsValidInput;
+				return bIsInputParameter;
 			}
 		};
 
@@ -245,6 +254,19 @@ namespace UE::KismetCompiler::Private
 								*UserPin->PinType.PinCategory.ToString());
 						}
 
+						// Subpins aren't a part of the function signature, so we need to skip them.
+						while (PinCursor < Node->Pins.Num())
+						{
+							UEdGraphPin* Pin = Node->Pins[PinCursor];
+							check(Pin);
+							if (Pin->ParentPin == nullptr)
+							{
+								break;
+							}
+
+							++PinCursor;
+						}
+
 						if (PinCursor < Node->Pins.Num())
 						{
 							UEdGraphPin* Pin = Node->Pins[PinCursor];
