Expose on spawn properties should use BlueprintSetter if available

diff --git a/Engine/Source/Editor/KismetCompiler/Private/KismetCompilerMisc.cpp b/Engine/Source/Editor/KismetCompiler/Private/KismetCompilerMisc.cpp
--- a/Engine/Source/Editor/KismetCompiler/Private/KismetCompilerMisc.cpp
+++ b/Engine/Source/Editor/KismetCompiler/Private/KismetCompilerMisc.cpp
@@ -754,16 +754,16 @@ UEdGraphPin* FKismetCompilerUtilities::GenerateAssignmentNodes(class FKismetComp
 		if (NULL == CallBeginSpawnNode->FindPin(OrgPin->PinName) &&
 			(OrgPin->LinkedTo.Num() > 0 || bHasDefaultValue))
 		{
-			if( OrgPin->LinkedTo.Num() == 0 )
+			FProperty* Property = FindFProperty<FProperty>(ForClass, OrgPin->PinName);
+			// NULL property indicates that this pin was part of the original node, not the 
+			// class we're assigning to:
+			if (!Property)
 			{
-				FProperty* Property = FindFProperty<FProperty>(ForClass, OrgPin->PinName);
-				// NULL property indicates that this pin was part of the original node, not the 
-				// class we're assigning to:
-				if( !Property )
-				{
-					continue;
-				}
+				continue;
+			}
 
+			if( OrgPin->LinkedTo.Num() == 0 )
+			{
 				// We don't want to generate an assignment node unless the default value 
 				// differs from the value in the CDO:
 				FString DefaultValueAsString;
@@ -786,8 +786,41 @@ UEdGraphPin* FKismetCompilerUtilities::GenerateAssignmentNodes(class FKismetComp
 				}
 			}
 
-			UFunction* SetByNameFunction = Schema->FindSetVariableByNameFunction(OrgPin->PinType);
-			if (SetByNameFunction)
+			const FString& SetFunctionName = Property->GetMetaData(FBlueprintMetadata::MD_PropertySetFunction);
+			if (!SetFunctionName.IsEmpty())
+			{
+				UClass* OwnerClass = Property->GetOwnerClass();
+				UFunction* SetFunction = OwnerClass->FindFunctionByName(*SetFunctionName);
+				check(SetFunction);
+
+				UK2Node_CallFunction* CallFuncNode = CompilerContext.SpawnIntermediateNode<UK2Node_CallFunction>(SpawnNode, SourceGraph);
+				CallFuncNode->SetFromFunction(SetFunction);
+				CallFuncNode->AllocateDefaultPins();
+
+				// Connect this node into the exec chain
+				Schema->TryCreateConnection(LastThen, CallFuncNode->GetExecPin());
+				LastThen = CallFuncNode->GetThenPin();
+
+				// Connect the new object to the 'object' pin
+				UEdGraphPin* ObjectPin = Schema->FindSelfPin(*CallFuncNode, EGPD_Input);
+				CallBeginResult->MakeLinkTo(ObjectPin);
+
+				// Move Value pin connections
+				UEdGraphPin* SetFunctionValuePin = nullptr;
+				for (UEdGraphPin* CallFuncPin : CallFuncNode->Pins)
+				{
+					if (!Schema->IsMetaPin(*CallFuncPin))
+					{
+						check(CallFuncPin->Direction == EGPD_Input);
+						SetFunctionValuePin = CallFuncPin;
+						break;
+					}
+				}
+				check(SetFunctionValuePin);
+
+				CompilerContext.MovePinLinksToIntermediate(*OrgPin, *SetFunctionValuePin);
+			}
+			else if (UFunction* SetByNameFunction = Schema->FindSetVariableByNameFunction(OrgPin->PinType))
 			{
 				UK2Node_CallFunction* SetVarNode = nullptr;
 				if (OrgPin->PinType.IsArray())
@@ -805,7 +838,7 @@ UEdGraphPin* FKismetCompilerUtilities::GenerateAssignmentNodes(class FKismetComp
 				Schema->TryCreateConnection(LastThen, SetVarNode->GetExecPin());
 				LastThen = SetVarNode->GetThenPin();
 
-				// Connect the new actor to the 'object' pin
+				// Connect the new object to the 'object' pin
 				UEdGraphPin* ObjectPin = SetVarNode->FindPinChecked(ObjectParamName);
 				CallBeginResult->MakeLinkTo(ObjectPin);
 
@@ -851,7 +884,6 @@ UEdGraphPin* FKismetCompilerUtilities::GenerateAssignmentNodes(class FKismetComp
 						CompilerContext.MovePinLinksToIntermediate(*OrgPin, *ValuePin);
 						SetVarNode->PinConnectionListChanged(ValuePin);
 					}
-
 				}
 			}
 		}
