Change default behavior for variable name collisions with a native parent class

diff --git a/Engine/Source/Editor/KismetCompiler/Private/KismetCompiler.cpp b/Engine/Source/Editor/KismetCompiler/Private/KismetCompiler.cpp
--- a/Engine/Source/Editor/KismetCompiler/Private/KismetCompiler.cpp
+++ b/Engine/Source/Editor/KismetCompiler/Private/KismetCompiler.cpp
@@ -918,6 +918,7 @@ void FKismetCompilerContext::ValidateVariableNames()
 			ParentBPNameValidator = MakeShareable(new FKismetNameValidator(ParentBP));
 		}
 
+		TArray<FName> RemoveVars;
 		for (FBPVariableDescription& VarDesc : Blueprint->NewVariables)
 		{
 			FName OldVarName = VarDesc.VarName;
@@ -930,11 +931,36 @@ void FKismetCompilerContext::ValidateVariableNames()
 			}
 			else if (ParentClass->IsNative()) // the above case handles when the parent is a blueprint
 			{
-				FFieldVariant ExisingField = FindUFieldOrFProperty(ParentClass, *VarNameStr);
-				if (ExisingField)
+				FFieldVariant ExistingField = FindUFieldOrFProperty(ParentClass, *VarNameStr);
+				if (ExistingField)
 				{
 					UE_LOG(LogK2Compiler, Warning, TEXT("ValidateVariableNames name %s (used in %s) is already taken by %s")
-						, *VarNameStr, *Blueprint->GetPathName(), *ExisingField.GetPathName());
+						, *VarNameStr, *Blueprint->GetPathName(), *ExistingField.GetPathName());
+
+					// If type is equivalent, we'll remove the blueprint variable and use the native variable wherever it's referenced
+					if (Blueprint->GeneratedClass)
+					{
+						FFieldVariant ExistingBPClassField = FindUFieldOrFProperty(Blueprint->GeneratedClass, *VarNameStr);
+						bool bIsCompatible = false;
+						if (ExistingBPClassField)
+						{
+							if (const UObject* AsObject = ExistingField.ToUObject())
+							{
+								bIsCompatible = ExistingBPClassField.IsA(AsObject->GetClass());
+							}
+							else if (const FField* AsField = ExistingField.ToField())
+							{
+								bIsCompatible = ExistingBPClassField.IsA(AsField->GetClass());
+							}
+						}
+
+						if (bIsCompatible)
+						{
+							RemoveVars.Add(OldVarName);
+							continue;
+						}
+					}
+
 					NewVarName = FBlueprintEditorUtils::FindUniqueKismetName(Blueprint, VarNameStr);
 				}
 			}
@@ -952,6 +978,22 @@ void FKismetCompilerContext::ValidateVariableNames()
 				FBlueprintEditorUtils::RenameMemberVariable(Blueprint, OldVarName, NewVarName);
 			}
 		}
+
+		for (FName VarName : RemoveVars)
+		{
+			MessageLog.Note(
+				*FText::Format(
+					LOCTEXT("MemberVariableNativeConflictNoteFmt", "Found a member variable with a conflicting name ({0}) - replaced references with native parent class variable."),
+					FText::FromName(VarName)
+				).ToString()
+			);
+
+			int32 VarIndex = FBlueprintEditorUtils::FindNewVariableIndex(Blueprint, VarName);
+			if (VarIndex != INDEX_NONE)
+			{
+				Blueprint->NewVariables.RemoveAt(VarIndex);
+			}
+		}
 	}
 }
 
