Fix Blueprint compiler to generate a well-formed bytecode expression for a struct literal term when the struct type declares one or more fields of a type that does not have Blueprint support.

Other changes:
- Don't emit an ICE for the binary compatibility test between a literal term type and its associated coerce property when the property's underlying type is not supported; this is an undefined context.
- Include additional context in the error message that's output as a result of encountering an undefined opcode during script execution. Note that the callstack will identify the Blueprint class/function.

diff --git a/Engine/Source/Editor/KismetCompiler/Private/KismetCompilerVMBackend.cpp b/Engine/Source/Editor/KismetCompiler/Private/KismetCompilerVMBackend.cpp
--- a/Engine/Source/Editor/KismetCompiler/Private/KismetCompilerVMBackend.cpp
+++ b/Engine/Source/Editor/KismetCompiler/Private/KismetCompilerVMBackend.cpp
@@ -616,8 +616,8 @@ public:
 			// Additional Validation, since we cannot trust custom k2nodes
 			if (CoerceProperty && ensure(Schema) && ensure(CurrentCompilerContext))
 			{
-			    const bool bSecialCaseSelf = (Term->Type.PinSubCategory == UEdGraphSchema_K2::PN_Self);
-				if(!bSecialCaseSelf)
+			    const bool bSpecialCaseSelf = (Term->Type.PinSubCategory == UEdGraphSchema_K2::PN_Self);
+				if(!bSpecialCaseSelf)
 			    {
 				    FEdGraphPinType TrueType;
 				    const bool bValidProperty = Schema->ConvertPropertyToPinType(CoerceProperty, TrueType);
@@ -642,12 +642,12 @@ public:
 					    }
 					    return true;
 				    };
-    
-				    if (!bValidProperty || !AreTypesBinaryCompatible(Term->Type, TrueType))
-				    {
-					    const FString ErrorMessage = FString::Printf(TEXT("ICE: The type of property %s doesn't match a term. @@"), *CoerceProperty->GetPathName());
-					    CurrentCompilerContext->MessageLog.Error(*ErrorMessage, Term->SourcePin);
-				    }
+
+					if (bValidProperty && !AreTypesBinaryCompatible(Term->Type, TrueType))
+					{
+						const FString ErrorMessage = FString::Printf(TEXT("ICE: The type of property %s doesn't match the terminal type for pin @@."), *CoerceProperty->GetPathName());
+						CurrentCompilerContext->MessageLog.Error(*ErrorMessage, Term->SourcePin);
+					}
 				}
 			}
 
@@ -913,11 +913,17 @@ public:
 							continue;
 						}
 
+						// Create a new term for each property, and serialize it out
 						for (int32 ArrayIter = 0; ArrayIter < Prop->ArrayDim; ++ArrayIter)
 						{
-							// Create a new term for each property, and serialize it out
 							FBPTerminal NewTerm;
-							Schema->ConvertPropertyToPinType(Prop, NewTerm.Type);
+							if(!Schema->ConvertPropertyToPinType(Prop, NewTerm.Type))
+							{								
+								// Do nothing for unsupported/unhandled property types. This will leave the value unchanged from its constructed default.
+								Writer << EX_Nothing;
+								continue;
+							}
+
 							NewTerm.bIsLiteral = true;
 							NewTerm.Source = Term->Source;
 							NewTerm.SourcePin = Term->SourcePin;
@@ -1171,11 +1177,18 @@ public:
 			{
 				// Create a new term for each property, and serialize it out
 				FBPTerminal NewTerm;
-				Schema->ConvertPropertyToPinType(Prop, NewTerm.Type);
-				NewTerm.bIsLiteral = true;
-				Prop->ExportText_InContainer(0, NewTerm.Name, StructData, StructData, NULL, PPF_None);
+				if(Schema->ConvertPropertyToPinType(Prop, NewTerm.Type))
+				{
+					NewTerm.bIsLiteral = true;
+					Prop->ExportText_InContainer(0, NewTerm.Name, StructData, StructData, NULL, PPF_None);
 
-				EmitTermExpr(&NewTerm, Prop);
+					EmitTermExpr(&NewTerm, Prop);
+				}
+				else
+				{
+					// Do nothing for unsupported/unhandled property types. This will leave the value unchanged from its constructed default.
+					Writer << EX_Nothing;
+				}
 			}
 		}
 
diff --git a/Engine/Source/Runtime/CoreUObject/Private/UObject/ScriptCore.cpp b/Engine/Source/Runtime/CoreUObject/Private/UObject/ScriptCore.cpp
--- a/Engine/Source/Runtime/CoreUObject/Private/UObject/ScriptCore.cpp
+++ b/Engine/Source/Runtime/CoreUObject/Private/UObject/ScriptCore.cpp
@@ -2040,7 +2040,13 @@ void UObject::ProcessEvent( UFunction* Function, void* Parms )
 
 DEFINE_FUNCTION(UObject::execUndefined)
 {
-	Stack.Logf(ELogVerbosity::Error, TEXT("Unknown code token %02X"), Stack.Code[-1] );
+	const FText LocalizedErrorMessage = FText::Format(
+		LOCTEXT("UndefinedOpcode", "Encountered an undefined opcode ({0}) at byte offset {1}. The compiler may have generated an instruction sequence that was unexpected or incomplete."),
+		FText::FromString(FString::Printf(TEXT("0x%02X"), Stack.Code[-1])),
+		FText::AsNumber(Stack.Node ? reinterpret_cast<ScriptPointerType>(&Stack.Code[-1]) - reinterpret_cast<ScriptPointerType>(&Stack.Node->Script[0]) : 0)
+	);
+
+	Stack.Log(ELogVerbosity::Error, LocalizedErrorMessage.ToString());
 }
 
 DEFINE_FUNCTION(UObject::execLocalVariable)
