Reduce to a stub function if the BP compiler backend raises an error during translation of compiled script statements to bytecode.

diff --git a/Engine/Source/Editor/KismetCompiler/Private/KismetCompilerVMBackend.cpp b/Engine/Source/Editor/KismetCompiler/Private/KismetCompilerVMBackend.cpp
--- a/Engine/Source/Editor/KismetCompiler/Private/KismetCompilerVMBackend.cpp
+++ b/Engine/Source/Editor/KismetCompiler/Private/KismetCompilerVMBackend.cpp
@@ -2355,20 +2355,34 @@ void FKismetCompilerVMBackend::ConstructFunction(FKismetFunctionContext& Functio
 			UEdGraphNode* StatementNode = FunctionContext.LinearExecutionList[NodeIndex];
 			TArray<FBlueprintCompiledStatement*>* StatementList = FunctionContext.StatementsPerNode.Find(StatementNode);
 
-			if (StatementList != NULL)
+			if (StatementList != nullptr)
 			{
 				for (int32 StatementIndex = 0; StatementIndex < StatementList->Num(); ++StatementIndex)
 				{
 					FBlueprintCompiledStatement* Statement = (*StatementList)[StatementIndex];
 
 					ScriptWriter.GenerateCodeForStatement(CompilerContext, FunctionContext, *Statement, StatementNode);
+
+					// Abort code generation on error (no need to process additional statements).
+					if (FunctionContext.MessageLog.NumErrors > 0)
+					{
+						break;
+					}
 				}
 			}
+
+			// Reduce to a stub if any errors were raised. This ensures the VM won't attempt to evaluate an incomplete expression.
+			if (FunctionContext.MessageLog.NumErrors > 0)
+			{
+				ScriptArray.Empty();
+				ReturnStatement.bIsJumpTarget = false;
+				break;
+			}
 		}
 	}
 
 	// Handle the function return value
-	ScriptWriter.GenerateCodeForStatement(CompilerContext, FunctionContext, ReturnStatement, NULL);	
+	ScriptWriter.GenerateCodeForStatement(CompilerContext, FunctionContext, ReturnStatement, nullptr);	
 
 	// Fix up jump addresses
 	ScriptWriter.PerformFixups();
