UE-151338 - Implement refactored implicit conversion for math types.

Removed container and struct conversions from the script VM. This introduced complexity that the VM doesn't need, nor did it scale for the various struct types that we want to implicitly convert in Blueprints. Instead, the script VM is *only* aware of float<->double conversion. Container and struct conversions have now been moved to BlueprintTypeConversions. Currently, only the existing FVector3f<->FVector3d conversion has been added, but the remaining LWC types will be added in a subsequent change.

During Blueprint compilation of container and struct conversions, we now inject a function call into the bytecode that performs the conversion, which is better suited to the task instead of burdening the VM with the work. One drawback to this technique is that containers are slightly more inefficient when it comes to conversions. They won't know their type(s) ahead of time, which requires dynamically looking up a conversion function at runtime. We can possibly optimize this further, but the generaly recommendation is to avoid implicit conversions of container types when possible.

Additionally, a couple of convenience functions were added to the KismetCastingUtils to help remove a fair amount of boilerplate code that was used for implicit casting in various node types. ScriptCastingUtils.h was also removed since the VM no longer needs to concern itself with complex conversions.

diff --git a/Engine/Source/Developer/ScriptDisassembler/Private/ScriptDisassembler.cpp b/Engine/Source/Developer/ScriptDisassembler/Private/ScriptDisassembler.cpp
--- a/Engine/Source/Developer/ScriptDisassembler/Private/ScriptDisassembler.cpp
+++ b/Engine/Source/Developer/ScriptDisassembler/Private/ScriptDisassembler.cpp
@@ -251,21 +251,7 @@ void FKismetBytecodeDisassembler::ProcessCommon(int32& ScriptIndex, EExprToken O
 		TEXT("ObjectToBool"),
 		TEXT("InterfaceToBool"),
 		TEXT("DoubleToFloat"),
-		TEXT("DoubleToFloatArray"),
-		TEXT("DoubleToFloatSet"),
 		TEXT("FloatToDouble"),
-		TEXT("FloatToDoubleArray"),
-		TEXT("FloatToDoubleSet"),
-		TEXT("VectorToVector3f"),
-		TEXT("Vector3fToVector"),
-		TEXT("FloatToDoubleKeys_Map"),
-		TEXT("DoubleToFloatKeys_Map"),
-		TEXT("FloatToDoubleValues_Map"),
-		TEXT("DoubleToFloatValues_Map"),
-		TEXT("FloatToDoubleKeys_FloatToDoubleValues_Map"),
-		TEXT("DoubleToFloatKeys_FloatToDoubleValues_Map"),
-		TEXT("DoubleToFloatKeys_DoubleToFloatValues_Map"),
-		TEXT("FloatToDoubleKeys_DoubleToFloatValues_Map"),
 	};
 
 	auto PrintVariable = [&ScriptIndex, Opcode, this](FStringView VariableDescription)
diff --git a/Engine/Source/Editor/BlueprintGraph/Private/CallFunctionHandler.cpp b/Engine/Source/Editor/BlueprintGraph/Private/CallFunctionHandler.cpp
--- a/Engine/Source/Editor/BlueprintGraph/Private/CallFunctionHandler.cpp
+++ b/Engine/Source/Editor/BlueprintGraph/Private/CallFunctionHandler.cpp
@@ -75,6 +75,8 @@ public:
  */
 void FKCHandler_CallFunction::CreateFunctionCallStatement(FKismetFunctionContext& Context, UEdGraphNode* Node, UEdGraphPin* SelfPin)
 {
+	using namespace UE::KismetCompiler;
+
 	int32 NumErrorsAtStart = CompilerContext.MessageLog.NumErrors;
 
 	// Find the function, starting at the parent class
@@ -162,8 +164,8 @@ void FKCHandler_CallFunction::CreateFunctionCallStatement(FKismetFunctionContext
 		//     NativeFunctionWithReferenceParam(CastedInput)	; CastedInput has possibly changed since this function takes a float&
 		//     Input = (double)CastedInput						; Now we need to propagate that change back to Input
 		//
-		using CastEntry = TTuple<FBPTerminal*, FBPTerminal*, EKismetCompiledStatementType>;
-		TArray<CastEntry> ModifiedCastInputs;
+		using CastEntryT = TPair<FBPTerminal*, CastingUtils::FImplicitCastParams>;
+		TArray<CastEntryT> ModifiedCastInputs;
 
 		// Check each property
 		bool bMatchedAllParams = true;
@@ -254,14 +256,14 @@ void FKCHandler_CallFunction::CreateFunctionCallStatement(FKismetFunctionContext
 								}
 
 								{
-									using namespace UE::KismetCompiler;
-
-									TOptional<TPair<FBPTerminal*, EKismetCompiledStatementType>> ImplicitCastEntry =
-										CastingUtils::InsertImplicitCastStatement(Context, PinMatch, RHSTerm);
+									const CastingUtils::FImplicitCastParams* CastParams =
+										Context.ImplicitCastMap.Find(PinMatch);
 
-									if (ImplicitCastEntry)
+									if (CastParams)
 									{
-										FBPTerminal* CastTerm = ImplicitCastEntry->Get<0>();
+										FBPTerminal* ImplicitCastTerm =
+											CastingUtils::InsertImplicitCastStatement(Context, PinMatch, RHSTerm);
+										check(ImplicitCastTerm);
 
 										bool bIsNonConstReference =
 											Property->HasAllPropertyFlags(CPF_OutParm | CPF_ReferenceParm) &&
@@ -269,15 +271,23 @@ void FKCHandler_CallFunction::CreateFunctionCallStatement(FKismetFunctionContext
 
 										if (bIsNonConstReference)
 										{
-											TOptional<EKismetCompiledStatementType> CastType = 
-												CastingUtils::GetInverseCastStatement(ImplicitCastEntry->Get<1>());
+											CastingUtils::FImplicitCastParams InverseCastParams = *CastParams;
 
-											check(CastType.IsSet());
+											if (CastParams->Conversion.Type == CastingUtils::FloatingPointCastType::FloatToDouble)
+											{
+												InverseCastParams.Conversion.Type = CastingUtils::FloatingPointCastType::DoubleToFloat;
+											}
+											else if (CastParams->Conversion.Type == CastingUtils::FloatingPointCastType::DoubleToFloat)
+											{
+												InverseCastParams.Conversion.Type = CastingUtils::FloatingPointCastType::FloatToDouble;
+											}
+
+											InverseCastParams.TargetTerminal = RHSTerm;
 
-											ModifiedCastInputs.Add(CastEntry{RHSTerm, CastTerm, *CastType});
+											ModifiedCastInputs.Add(CastEntryT{ImplicitCastTerm, InverseCastParams});
 										}
 
-										RHSTerm = CastTerm;
+										RHSTerm = ImplicitCastTerm;
 									}
 								}
 
@@ -505,21 +515,12 @@ void FKCHandler_CallFunction::CreateFunctionCallStatement(FKismetFunctionContext
 			}
 
 			{
-				using namespace UE::KismetCompiler;
-
 				for (const auto& It : ModifiedCastInputs)
 				{
-					FBPTerminal* LocalLHSTerm = It.Get<0>();
-					FBPTerminal* LocalRHSTerm = It.Get<1>();
-					EKismetCompiledStatementType LocalCastType = It.Get<2>();
-
-					check(LocalLHSTerm);
-					check(LocalRHSTerm);
+					FBPTerminal* LocalRHSTerm = It.Get<0>();
+					CastingUtils::FImplicitCastParams LocalInverseCastParams = It.Get<1>();
 
-					FBlueprintCompiledStatement& CastStatement = Context.AppendStatementForNode(Node);
-					CastStatement.LHS = LocalLHSTerm;
-					CastStatement.Type = LocalCastType;
-					CastStatement.RHS.Add(LocalRHSTerm);
+					CastingUtils::InsertImplicitCastStatement(Context, LocalInverseCastParams, LocalRHSTerm);
 				}
 			}
 
diff --git a/Engine/Source/Editor/BlueprintGraph/Private/EdGraphSchema_K2.cpp b/Engine/Source/Editor/BlueprintGraph/Private/EdGraphSchema_K2.cpp
--- a/Engine/Source/Editor/BlueprintGraph/Private/EdGraphSchema_K2.cpp
+++ b/Engine/Source/Editor/BlueprintGraph/Private/EdGraphSchema_K2.cpp
@@ -67,9 +67,9 @@
 #include "Kismet2/BlueprintEditorUtils.h"
 #include "Settings/EditorStyleSettings.h"
 #include "Editor.h"
-
 #include "Kismet/BlueprintMapLibrary.h"
 #include "Kismet/BlueprintSetLibrary.h"
+#include "Kismet/BlueprintTypeConversions.h"
 #include "Kismet/KismetArrayLibrary.h"
 #include "Kismet/KismetMathLibrary.h"
 #include "GraphEditorActions.h"
@@ -4379,6 +4379,8 @@ bool UEdGraphSchema_K2::DefaultValueSimpleValidation(const FEdGraphPinType& PinT
 
 bool UEdGraphSchema_K2::ArePinTypesCompatible(const FEdGraphPinType& Output, const FEdGraphPinType& Input, const UClass* CallingContext, bool bIgnoreArray /*= false*/) const
 {
+	using namespace UE::Kismet::BlueprintTypeConversions;
+
 	if (!bIgnoreArray && 
 		(Output.ContainerType != Input.ContainerType) && 
 		(Input.PinCategory != PC_Wildcard || Input.IsContainer()) && 
@@ -4388,9 +4390,14 @@ bool UEdGraphSchema_K2::ArePinTypesCompatible(const FEdGraphPinType& Output, con
 	}
 	else if (Output.PinCategory == Input.PinCategory)
 	{
-		bool bAreConvertibleVectorTypes =
-			((Output.PinSubCategoryObject == VectorStruct) && (Input.PinSubCategoryObject == Vector3fStruct)) ||
-			((Output.PinSubCategoryObject == Vector3fStruct) && (Input.PinSubCategoryObject == VectorStruct));
+		bool bAreConvertibleStructs = false;
+		const UScriptStruct* OutputStruct = Cast<UScriptStruct>(Output.PinSubCategoryObject.Get());
+		const UScriptStruct* InputStruct = Cast<UScriptStruct>(Input.PinSubCategoryObject.Get());
+		if (OutputStruct != InputStruct)
+		{
+			bAreConvertibleStructs =
+				FStructConversionTable::Get().GetConversionFunction(OutputStruct, InputStruct).IsSet();
+		}
 
 		if ((Output.PinSubCategory == Input.PinSubCategory) 
 			&& (Output.PinSubCategoryObject == Input.PinSubCategoryObject)
@@ -4398,10 +4405,19 @@ bool UEdGraphSchema_K2::ArePinTypesCompatible(const FEdGraphPinType& Output, con
 		{
 			if(Input.IsMap())
 			{
+				OutputStruct = Cast<UScriptStruct>(Output.PinValueType.TerminalSubCategoryObject.Get());
+				InputStruct = Cast<UScriptStruct>(Input.PinValueType.TerminalSubCategoryObject.Get());
+				if (OutputStruct != InputStruct)
+				{
+					bAreConvertibleStructs =
+						FStructConversionTable::Get().GetConversionFunction(OutputStruct, InputStruct).IsSet();
+				}
+
 				return 
 					Input.PinValueType.TerminalCategory == PC_Wildcard ||
 					Output.PinValueType.TerminalCategory == PC_Wildcard ||
 					((Input.PinValueType.TerminalCategory == PC_Real) && (Output.PinValueType.TerminalCategory == PC_Real)) ||
+					bAreConvertibleStructs ||
 					Input.PinValueType == Output.PinValueType;
 			}
 			return true;
@@ -4412,7 +4428,7 @@ bool UEdGraphSchema_K2::ArePinTypesCompatible(const FEdGraphPinType& Output, con
 		{
 			return true;
 		}
-		else if (bAreConvertibleVectorTypes)
+		else if (bAreConvertibleStructs)
 		{
 			return true;
 		}
diff --git a/Engine/Source/Editor/BlueprintGraph/Private/K2Node_AssignmentStatement.cpp b/Engine/Source/Editor/BlueprintGraph/Private/K2Node_AssignmentStatement.cpp
--- a/Engine/Source/Editor/BlueprintGraph/Private/K2Node_AssignmentStatement.cpp
+++ b/Engine/Source/Editor/BlueprintGraph/Private/K2Node_AssignmentStatement.cpp
@@ -50,26 +50,14 @@ public:
 
 		if (VariablePinNet && ValuePinNet)
 		{
-			TOptional<CastingUtils::StatementNamePair> ConversionType =
-				CastingUtils::GetFloatingPointConversionType(*ValuePinNet, *VariablePinNet);
+			CastingUtils::FConversion Conversion =
+				CastingUtils::GetFloatingPointConversion(*ValuePinNet, *VariablePinNet);
 
-			if (ConversionType)
+			if (Conversion.Type != CastingUtils::FloatingPointCastType::None)
 			{
-				FString DescriptiveName = Node->GetName();
+				FBPTerminal* NewTerm = CastingUtils::MakeImplicitCastTerminal(Context, VariablePinNet, Node);
 
-				FString TerminalName = FString::Printf(TEXT("%s_%s_%s"),
-					*DescriptiveName,
-					*VariablePinNet->PinName.ToString(),
-					ConversionType->Get<1>());
-
-				FBPTerminal* NewTerm = Context.CreateLocalTerminal();
-				NewTerm->Name = TerminalName;
-				NewTerm->Type = VariablePinNet->PinType;
-				NewTerm->Source = Node;
-
-				EKismetCompiledStatementType CastType = ConversionType->Get<0>();
-
-				Context.ImplicitCastMap.Add(VariablePin, FImplicitCastParams{CastType, NewTerm, Node});
+				Context.ImplicitCastMap.Add(VariablePin, CastingUtils::FImplicitCastParams{Conversion, NewTerm, Node});
 			}
 		}
 		else
diff --git a/Engine/Source/Editor/BlueprintGraph/Private/K2Node_MakeContainer.cpp b/Engine/Source/Editor/BlueprintGraph/Private/K2Node_MakeContainer.cpp
--- a/Engine/Source/Editor/BlueprintGraph/Private/K2Node_MakeContainer.cpp
+++ b/Engine/Source/Editor/BlueprintGraph/Private/K2Node_MakeContainer.cpp
@@ -51,12 +51,12 @@ void FKCHandler_MakeContainer::Compile(FKismetFunctionContext& Context, UEdGraph
 				{
 					using namespace UE::KismetCompiler;
 
-					TOptional<TPair<FBPTerminal*, EKismetCompiledStatementType>> ImplicitCastEntry =
+					FBPTerminal* ImplicitCastEntry =
 						CastingUtils::InsertImplicitCastStatement(Context, Pin, RHSTerm);
 
 					if (ImplicitCastEntry)
 					{
-						RHSTerm = ImplicitCastEntry->Get<0>();
+						RHSTerm = ImplicitCastEntry;
 					}
 				}
 				RHSTerms.Add(RHSTerm);
diff --git a/Engine/Source/Editor/BlueprintGraph/Private/K2Node_Select.cpp b/Engine/Source/Editor/BlueprintGraph/Private/K2Node_Select.cpp
--- a/Engine/Source/Editor/BlueprintGraph/Private/K2Node_Select.cpp
+++ b/Engine/Source/Editor/BlueprintGraph/Private/K2Node_Select.cpp
@@ -152,12 +152,12 @@ public:
 				{
 					using namespace UE::KismetCompiler;
 
-					TOptional<TPair<FBPTerminal*, EKismetCompiledStatementType>> ImplicitCastEntry =
+					FBPTerminal* ImplicitCastEntry =
 						CastingUtils::InsertImplicitCastStatement(Context, OptionPins[OptionIdx], ValueTerm);
 
 					if (ImplicitCastEntry)
 					{
-						ValueTerm = ImplicitCastEntry->Get<0>();
+						ValueTerm = ImplicitCastEntry;
 					}
 				}
 
diff --git a/Engine/Source/Editor/BlueprintGraph/Private/K2Node_VariableSetRef.cpp b/Engine/Source/Editor/BlueprintGraph/Private/K2Node_VariableSetRef.cpp
--- a/Engine/Source/Editor/BlueprintGraph/Private/K2Node_VariableSetRef.cpp
+++ b/Engine/Source/Editor/BlueprintGraph/Private/K2Node_VariableSetRef.cpp
@@ -45,19 +45,16 @@ public:
 
 			if ((VariablePinNet != nullptr) && (ValuePinNet != nullptr))
 			{
-				TOptional<CastingUtils::StatementNamePair> ConversionType =
-					CastingUtils::GetFloatingPointConversionType(*ValuePinNet, *VariablePinNet);
+				CastingUtils::FConversion Conversion =
+					CastingUtils::GetFloatingPointConversion(*ValuePinNet, *VariablePinNet);
 
-				if (ConversionType)
+				if (Conversion.Type != CastingUtils::FloatingPointCastType::None)
 				{
 					check(!ImplicitCastMap.Contains(VarRefNode));
 
-					FBPTerminal* NewTerminal = Context.CreateLocalTerminal();
-					UEdGraphNode* OwningNode = VariablePinNet->GetOwningNode();
-					NewTerminal->CopyFromPin(VariablePinNet, Context.NetNameMap->MakeValidName(VariablePin, ConversionType->Get<1>()));
-					NewTerminal->Source = OwningNode;
+					FBPTerminal* NewTerminal = CastingUtils::MakeImplicitCastTerminal(Context, VariablePinNet);
 
-					ImplicitCastMap.Add(VarRefNode, TPair<FBPTerminal*, EKismetCompiledStatementType>{ NewTerminal, ConversionType->Get<0>() });
+					ImplicitCastMap.Add(VarRefNode, CastingUtils::FImplicitCastParams{Conversion, NewTerminal, Node});
 				}
 			}
 		}
@@ -103,17 +100,11 @@ private:
 				using namespace UE::KismetCompiler;
 
 				UK2Node_VariableSetRef* VarRefNode = CastChecked<UK2Node_VariableSetRef>(Node);
-				if (TPair<FBPTerminal*, EKismetCompiledStatementType>* CastEntry = ImplicitCastMap.Find(VarRefNode))
+				if (CastingUtils::FImplicitCastParams* CastParams = ImplicitCastMap.Find(VarRefNode))
 				{
-					FBPTerminal* CastTerminal = CastEntry->Get<0>();
-					EKismetCompiledStatementType StatementType = CastEntry->Get<1>();
-
-					FBlueprintCompiledStatement& CastStatement = Context.AppendStatementForNode(Node);
-					CastStatement.LHS = CastTerminal;
-					CastStatement.Type = StatementType;
-					CastStatement.RHS.Add(RHSTerm);
-
-					RHSTerm = CastTerminal;
+					CastingUtils::InsertImplicitCastStatement(Context, *CastParams, RHSTerm);
+					
+					RHSTerm = CastParams->TargetTerminal;
 
 					ImplicitCastMap.Remove(VarRefNode);
 
@@ -146,7 +137,7 @@ private:
 		}
 	}
 
-	TMap<UK2Node_VariableSetRef*, TPair<FBPTerminal*, EKismetCompiledStatementType>> ImplicitCastMap;
+	TMap<UK2Node_VariableSetRef*, UE::KismetCompiler::CastingUtils::FImplicitCastParams> ImplicitCastMap;
 };
 
 UK2Node_VariableSetRef::UK2Node_VariableSetRef(const FObjectInitializer& ObjectInitializer)
diff --git a/Engine/Source/Editor/BlueprintGraph/Private/MathExpressionHandler.cpp b/Engine/Source/Editor/BlueprintGraph/Private/MathExpressionHandler.cpp
--- a/Engine/Source/Editor/BlueprintGraph/Private/MathExpressionHandler.cpp
+++ b/Engine/Source/Editor/BlueprintGraph/Private/MathExpressionHandler.cpp
@@ -131,30 +131,31 @@ FBlueprintCompiledStatement* FKCHandler_MathExpression::GenerateFunctionRPN(UEdG
 
 			if (RHSTerm)
 			{
-				const FImplicitCastParams* CastParams =
+				using namespace UE::KismetCompiler;
+
+				const CastingUtils::FImplicitCastParams* CastParams =
 					Context.ImplicitCastMap.Find(PinMatch);
+
 				if (CastParams)
 				{
-					using namespace UE::KismetCompiler;
-
 					check(CastParams->TargetTerminal);
 
 					UFunction* CastFunction = nullptr;
 
-					switch (CastParams->CastType)
+					switch (CastParams->Conversion.Type)
 					{
-					case KCST_DoubleToFloatCast:
+					case CastingUtils::FloatingPointCastType::DoubleToFloat:
 						CastFunction = 
 							UKismetMathLibrary::StaticClass()->FindFunctionByName(GET_FUNCTION_NAME_CHECKED(UKismetMathLibrary, Conv_DoubleToFloat));
 						break;
 
-					case KCST_FloatToDoubleCast:
+					case CastingUtils::FloatingPointCastType::FloatToDouble:
 						CastFunction = 
 							UKismetMathLibrary::StaticClass()->FindFunctionByName(GET_FUNCTION_NAME_CHECKED(UKismetMathLibrary, Conv_FloatToDouble));
 						break;
 
 					default:
-						checkf(false, TEXT("Unsupported cast type used in math expression node: %d"), CastParams->CastType);
+						checkf(false, TEXT("Unsupported cast type used in math expression node: %d"), CastParams->Conversion.Type);
 					}
 
 					check(CastFunction);
@@ -257,20 +258,15 @@ void FKCHandler_MathExpression::RegisterNets(FKismetFunctionContext& Context, UE
 					{
 						if (DestinationPin)
 						{
-							TOptional<CastingUtils::StatementNamePair> ConversionType =
-								CastingUtils::GetFloatingPointConversionType(*PinNet, *DestinationPin);
+							CastingUtils::FConversion Conversion =
+								CastingUtils::GetFloatingPointConversion(*PinNet, *DestinationPin);
 
-							if (ConversionType)
+							if (Conversion.Type != CastingUtils::FloatingPointCastType::None)
 							{
-								EKismetCompiledStatementType CastType = ConversionType->Get<0>();
-								const TCHAR* TermName = ConversionType->Get<1>();
-
-								FBPTerminal* NewTerm = Context.CreateLocalTerminal();
+								FBPTerminal* NewTerm = CastingUtils::MakeImplicitCastTerminal(Context, DestinationPin);
 								UEdGraphNode* OwningNode = DestinationPin->GetOwningNode();
-								NewTerm->CopyFromPin(DestinationPin, Context.NetNameMap->MakeValidName(DestinationPin, TermName));
-								NewTerm->Source = OwningNode;
 
-								Context.ImplicitCastMap.Add(DestinationPin, FImplicitCastParams{ CastType, NewTerm, OwningNode });
+								Context.ImplicitCastMap.Add(DestinationPin, CastingUtils::FImplicitCastParams{Conversion, NewTerm, OwningNode});
 							}
 						}
 					}
diff --git a/Engine/Source/Editor/KismetCompiler/Private/KismetCastingUtils.cpp b/Engine/Source/Editor/KismetCompiler/Private/KismetCastingUtils.cpp
--- a/Engine/Source/Editor/KismetCompiler/Private/KismetCastingUtils.cpp
+++ b/Engine/Source/Editor/KismetCompiler/Private/KismetCastingUtils.cpp
@@ -4,258 +4,34 @@
 
 #include "KismetCompiledFunctionContext.h"
 #include "KismetCompilerMisc.h"
-
-namespace UE::KismetCompiler::CastingUtils::Private
-{
-
-enum class EPinType : uint8
-{
-	Float,
-	FloatArray,
-	FloatSet,
-
-	Double,
-	DoubleArray,
-	DoubleSet,
-
-	Vector3f,
-	Vector3fArray,
-	Vector3fSet,
-
-	Vector,
-	VectorArray,
-	VectorSet,
-
-	// Maps are particularly sinister, and are categorized separately here.
-	// Keys and values are casted independently of one another.
-
-	FloatKeyOtherValueMap,
-	DoubleKeyOtherValueMap,
-	OtherKeyFloatValueMap,
-	OtherKeyDoubleValueMap,
-	FloatKeyFloatValueMap,
-	FloatKeyDoubleValueMap,
-	DoubleKeyFloatValueMap,
-	DoubleKeyDoubleValueMap,
-
-	Other,
-};
-
-EPinType GetPinType(const UEdGraphPin& Pin)
-{
-	static UScriptStruct* VectorStruct = TBaseStructure<FVector>::Get();
-	static UScriptStruct* Vector3fStruct = TVariantStructure<FVector3f>::Get();
-
-	if (Pin.PinType.IsMap())
-	{
-		if ((Pin.PinType.PinCategory == UEdGraphSchema_K2::PC_Real) && (Pin.PinType.PinSubCategory == UEdGraphSchema_K2::PC_Float))
-		{
-			if ((Pin.PinType.PinValueType.TerminalCategory == UEdGraphSchema_K2::PC_Real) && (Pin.PinType.PinValueType.TerminalSubCategory == UEdGraphSchema_K2::PC_Float))
-			{
-				return EPinType::FloatKeyFloatValueMap;
-			}
-			else if ((Pin.PinType.PinValueType.TerminalCategory == UEdGraphSchema_K2::PC_Real) && (Pin.PinType.PinValueType.TerminalSubCategory == UEdGraphSchema_K2::PC_Double))
-			{
-				return EPinType::FloatKeyDoubleValueMap;
-			}
-			else
-			{
-				return EPinType::FloatKeyOtherValueMap;
-			}
-		}
-		else if ((Pin.PinType.PinCategory == UEdGraphSchema_K2::PC_Real) && (Pin.PinType.PinSubCategory == UEdGraphSchema_K2::PC_Double))
-		{
-			if ((Pin.PinType.PinValueType.TerminalCategory == UEdGraphSchema_K2::PC_Real) && (Pin.PinType.PinValueType.TerminalSubCategory == UEdGraphSchema_K2::PC_Float))
-			{
-				return EPinType::DoubleKeyFloatValueMap;
-			}
-			else if ((Pin.PinType.PinValueType.TerminalCategory == UEdGraphSchema_K2::PC_Real) && (Pin.PinType.PinValueType.TerminalSubCategory == UEdGraphSchema_K2::PC_Double))
-			{
-				return EPinType::DoubleKeyDoubleValueMap;
-			}
-			else
-			{
-				return EPinType::DoubleKeyOtherValueMap;
-			}
-		}
-		else
-		{
-			if ((Pin.PinType.PinValueType.TerminalCategory == UEdGraphSchema_K2::PC_Real) && (Pin.PinType.PinValueType.TerminalSubCategory == UEdGraphSchema_K2::PC_Float))
-			{
-				return EPinType::OtherKeyFloatValueMap;
-			}
-			else if ((Pin.PinType.PinValueType.TerminalCategory == UEdGraphSchema_K2::PC_Real) && (Pin.PinType.PinValueType.TerminalSubCategory == UEdGraphSchema_K2::PC_Double))
-			{
-				return EPinType::OtherKeyDoubleValueMap;
-			}
-		}
-	}
-	else
-	{
-		if ((Pin.PinType.PinCategory == UEdGraphSchema_K2::PC_Real) && (Pin.PinType.PinSubCategory == UEdGraphSchema_K2::PC_Float))
-		{
-			if (Pin.PinType.IsArray())
-			{
-				return EPinType::FloatArray;
-			}
-			else if (Pin.PinType.IsSet())
-			{
-				return EPinType::FloatSet;
-			}
-			else
-			{
-				return EPinType::Float;
-			}
-		}
-		else if ((Pin.PinType.PinCategory == UEdGraphSchema_K2::PC_Real) && (Pin.PinType.PinSubCategory == UEdGraphSchema_K2::PC_Double))
-		{
-			if (Pin.PinType.IsArray())
-			{
-				return EPinType::DoubleArray;
-			}
-			else if (Pin.PinType.IsSet())
-			{
-				return EPinType::DoubleSet;
-			}
-			else
-			{
-				return EPinType::Double;
-			}
-		}
-		else if (Pin.PinType.PinCategory == UEdGraphSchema_K2::PC_Struct)
-		{
-			if (Pin.PinType.PinSubCategoryObject == Vector3fStruct)
-			{
-				if (Pin.PinType.IsArray())
-				{
-					return EPinType::Vector3fArray;
-				}
-				else if (Pin.PinType.IsSet())
-				{
-					return EPinType::Vector3fSet;
-				}
-				else
-				{
-					return EPinType::Vector3f;
-				}
-			}
-			else if (Pin.PinType.PinSubCategoryObject == VectorStruct)
-			{
-				if (Pin.PinType.IsArray())
-				{
-					return EPinType::VectorArray;
-				}
-				else if (Pin.PinType.IsSet())
-				{
-					return EPinType::VectorSet;
-				}
-				else
-				{
-					return EPinType::Vector;
-				}
-			}
-		}
-	}
-
-	return EPinType::Other;
-}
-
-} // namespace UE::KismetCompiler::CastingUtils::Private
+#include "Kismet/BlueprintTypeConversions.h"
 
 namespace UE::KismetCompiler::CastingUtils
 {
 
-TOptional<EKismetCompiledStatementType> GetInverseCastStatement(EKismetCompiledStatementType Statement)
+FBPTerminal* MakeImplicitCastTerminal(FKismetFunctionContext& Context, UEdGraphPin* Net, UEdGraphNode* SourceNode)
 {
-	TOptional<EKismetCompiledStatementType> Result;
-
-	switch (Statement)
-	{
-	case KCST_FloatToDoubleCast:
-		Result = KCST_DoubleToFloatCast;
-		break;
-	case KCST_FloatToDoubleArrayCast:
-		Result = KCST_DoubleToFloatArrayCast;
-		break;
-	case KCST_FloatToDoubleSetCast:
-		Result = KCST_DoubleToFloatSetCast;
-		break;
-
-	case KCST_DoubleToFloatCast:
-		Result = KCST_FloatToDoubleCast;
-		break;
-	case KCST_DoubleToFloatArrayCast:
-		Result = KCST_FloatToDoubleArrayCast;
-		break;
-	case KCST_DoubleToFloatSetCast:
-		Result = KCST_FloatToDoubleSetCast;
-		break;
-
-	case KCST_Vector3fToVectorCast:
-		Result = KCST_VectorToVector3fCast;
-		break;
-	case KCST_Vector3fToVectorArrayCast:
-		Result = KCST_VectorToVector3fArrayCast;
-		break;
-	case KCST_Vector3fToVectorSetCast:
-		Result = KCST_VectorToVector3fSetCast;
-		break;
-
-	case KCST_VectorToVector3fCast:
-		Result = KCST_Vector3fToVectorCast;
-		break;
-	case KCST_VectorToVector3fArrayCast:
-		Result = KCST_Vector3fToVectorArrayCast;
-		break;
-	case KCST_VectorToVector3fSetCast:
-		Result = KCST_Vector3fToVectorSetCast;
-		break;
+	check(Net);
+	
+	FBPTerminal* Result = Context.CreateLocalTerminal();
+	check(Result);
 
-	case KCST_FloatToDoubleKeys_MapCast:
-		Result = KCST_DoubleToFloatKeys_MapCast;
-		break;
-	case KCST_DoubleToFloatKeys_MapCast:
-		Result = KCST_FloatToDoubleKeys_MapCast;
-		break;
-	case KCST_FloatToDoubleValues_MapCast:
-		Result = KCST_DoubleToFloatValues_MapCast;
-		break;
-	case KCST_DoubleToFloatValues_MapCast:
-		Result = KCST_FloatToDoubleValues_MapCast;
-		break;
-
-	case KCST_FloatToDoubleKeys_FloatToDoubleValues_MapCast:
-		Result = KCST_DoubleToFloatKeys_DoubleToFloatValues_MapCast;
-		break;
-	case KCST_DoubleToFloatKeys_FloatToDoubleValues_MapCast:
-		Result = KCST_FloatToDoubleKeys_DoubleToFloatValues_MapCast;
-		break;
-	case KCST_DoubleToFloatKeys_DoubleToFloatValues_MapCast:
-		Result = KCST_FloatToDoubleKeys_FloatToDoubleValues_MapCast;
-		break;
-	case KCST_FloatToDoubleKeys_DoubleToFloatValues_MapCast:
-		Result = KCST_DoubleToFloatKeys_FloatToDoubleValues_MapCast;
-		break;
-	}
+	Result->CopyFromPin(Net, Context.NetNameMap->MakeValidName(Net, TEXT("ImplicitCast")));
+	Result->Source = (SourceNode ? SourceNode : Net->GetOwningNode());
 
 	return Result;
 }
 
 void RegisterImplicitCasts(FKismetFunctionContext& Context)
 {
-	using namespace UE::KismetCompiler::CastingUtils::Private;
-
-	auto AddCastMapping = [&Context](UEdGraphPin* DestinationPin, EKismetCompiledStatementType CastType, const TCHAR* TermName)
+	auto AddCastMapping = [&Context](UEdGraphPin* DestinationPin, const FConversion& Conversion)
 	{
 		check(DestinationPin);
-		check(TermName);
 
-		FBPTerminal* NewTerm = Context.CreateLocalTerminal();
+		FBPTerminal* NewTerm = MakeImplicitCastTerminal(Context, DestinationPin);
 		UEdGraphNode* OwningNode = DestinationPin->GetOwningNode();
-		NewTerm->CopyFromPin(DestinationPin, Context.NetNameMap->MakeValidName(DestinationPin, TermName));
-		NewTerm->Source = OwningNode;
 
-		Context.ImplicitCastMap.Add(DestinationPin, FImplicitCastParams{ CastType, NewTerm, OwningNode });
+		Context.ImplicitCastMap.Add(DestinationPin, FImplicitCastParams{Conversion, NewTerm, OwningNode});
 	};
 
 	// The current context's NetMap can be a mix of input and output pin types.
@@ -282,12 +58,10 @@ void RegisterImplicitCasts(FKismetFunctionContext& Context)
 					continue;
 				}
 
-				TOptional<StatementNamePair> ConversionType =
-					GetFloatingPointConversionType(*CurrentPin, *DestinationPin);
-
-				if (ConversionType)
+				FConversion Conversion = GetFloatingPointConversion(*CurrentPin, *DestinationPin);
+				if (Conversion.Type != FloatingPointCastType::None)
 				{
-					AddCastMapping(DestinationPin, ConversionType->Get<0>(), ConversionType->Get<1>());
+					AddCastMapping(DestinationPin, Conversion);
 				}
 			}
 		}
@@ -303,45 +77,72 @@ void RegisterImplicitCasts(FKismetFunctionContext& Context)
 				const UEdGraphPin* SourcePin = CurrentPin->LinkedTo[0];
 				check(SourcePin);
 
-				TOptional<StatementNamePair> ConversionType =
-					GetFloatingPointConversionType(*SourcePin, *CurrentPin);
-
-				if (ConversionType)
+				FConversion Conversion = GetFloatingPointConversion(*SourcePin, *CurrentPin);
+				if (Conversion.Type != FloatingPointCastType::None)
 				{
-					AddCastMapping(CurrentPin, ConversionType->Get<0>(), ConversionType->Get<1>());
+					AddCastMapping(CurrentPin, Conversion);
 				}
 			}
 		}
 	}
 }
 
-TOptional<TPair<FBPTerminal*, EKismetCompiledStatementType>>
-InsertImplicitCastStatement(FKismetFunctionContext& Context, UEdGraphPin* DestinationPin, FBPTerminal* RHSTerm)
+void InsertImplicitCastStatement(FKismetFunctionContext& Context, const FImplicitCastParams& CastParams, FBPTerminal* RHSTerm)
 {
-	using namespace UE::KismetCompiler::CastingUtils::Private;
+	check(RHSTerm);
+	check(CastParams.TargetTerminal);
+	check(CastParams.TargetNode);
+
+	EKismetCompiledStatementType CompiledStatementType = {};
+	UFunction* FunctionToCall = nullptr;
 
+	switch (CastParams.Conversion.Type)
+	{
+	case FloatingPointCastType::FloatToDouble:
+		CompiledStatementType = KCST_FloatToDoubleCast;
+		break;
+
+	case FloatingPointCastType::DoubleToFloat:
+		CompiledStatementType = KCST_DoubleToFloatCast;
+		break;
+
+	case FloatingPointCastType::Container:
+	case FloatingPointCastType::Struct:
+		CompiledStatementType = KCST_CallFunction;
+		FunctionToCall = CastParams.Conversion.Function;
+		check(FunctionToCall);
+		break;
+
+	default:
+		check(false);
+		break;
+	}
+
+	FBlueprintCompiledStatement& CastStatement = Context.AppendStatementForNode(CastParams.TargetNode);
+	CastStatement.Type = CompiledStatementType;
+	CastStatement.FunctionToCall = FunctionToCall;
+	CastStatement.LHS = CastParams.TargetTerminal;
+	CastStatement.RHS.Add(RHSTerm);
+}
+
+FBPTerminal* InsertImplicitCastStatement(FKismetFunctionContext& Context, UEdGraphPin* DestinationPin, FBPTerminal* RHSTerm)
+{
 	check(DestinationPin);
 	check(RHSTerm);
 
-	TOptional<TPair<FBPTerminal*, EKismetCompiledStatementType>> Result;
+	FBPTerminal* Result = nullptr;
 
 	const FImplicitCastParams* CastParams =
 		Context.ImplicitCastMap.Find(DestinationPin);
 
 	if (CastParams != nullptr)
 	{
-		check(CastParams->TargetTerminal);
-		check(CastParams->TargetNode);
-
-		FBlueprintCompiledStatement& CastStatement = Context.AppendStatementForNode(CastParams->TargetNode);
-		CastStatement.LHS = CastParams->TargetTerminal;
-		CastStatement.Type = CastParams->CastType;
-		CastStatement.RHS.Add(RHSTerm);
+		InsertImplicitCastStatement(Context, *CastParams, RHSTerm);
 
 		// Removal of the pin entry indicates to the compiler that the implicit cast has been processed.
 		Context.ImplicitCastMap.Remove(DestinationPin);
 
-		Result = TPair<FBPTerminal*, EKismetCompiledStatementType>{CastParams->TargetTerminal, CastParams->CastType};
+		Result = CastParams->TargetTerminal;
 	}
 
 	return Result;
@@ -356,65 +157,114 @@ bool RemoveRegisteredImplicitCast(FKismetFunctionContext& Context, const UEdGrap
 	return (RemovedCount > 0);
 }
 
-TOptional<StatementNamePair> GetFloatingPointConversionType(const UEdGraphPin& SourcePin, const UEdGraphPin& DestinationPin)
+FConversion GetFloatingPointConversion(const UEdGraphPin& SourcePin, const UEdGraphPin& DestinationPin)
 {
-	using namespace UE::KismetCompiler::CastingUtils::Private;
+	using namespace UE::Kismet::BlueprintTypeConversions;
 
-	using CastPair = TPair<EPinType, EPinType>;
+	UClass* BlueprintTypeConversionsClass = UBlueprintTypeConversions::StaticClass();
+	
+	UFunction* ArrayConversionFunction = BlueprintTypeConversionsClass->FindFunctionByName(TEXT("ConvertArrayType"));
+	check(ArrayConversionFunction);
 
-	static TMap<CastPair, StatementNamePair> ImplicitCastTable =
-	{
-		{ CastPair{EPinType::Float,						EPinType::Double},					StatementNamePair{KCST_FloatToDoubleCast,								TEXT("WideningCast")}		},
-		{ CastPair{EPinType::FloatArray,				EPinType::DoubleArray},				StatementNamePair{KCST_FloatToDoubleArrayCast,							TEXT("WideningArrayCast")}	},
-		{ CastPair{EPinType::FloatSet,					EPinType::DoubleSet},				StatementNamePair{KCST_FloatToDoubleSetCast,							TEXT("WideningSetCast")}	},
-
-		{ CastPair{EPinType::Double,					EPinType::Float},					StatementNamePair{KCST_DoubleToFloatCast,								TEXT("NarrowingCast")}		},
-		{ CastPair{EPinType::DoubleArray,				EPinType::FloatArray},				StatementNamePair{KCST_DoubleToFloatArrayCast,							TEXT("NarrowingArrayCast")}	},
-		{ CastPair{EPinType::DoubleSet,					EPinType::FloatSet},				StatementNamePair{KCST_DoubleToFloatSetCast,							TEXT("NarrowingSetCast")}	},
-
-		{ CastPair{EPinType::Vector3f,					EPinType::Vector},					StatementNamePair{KCST_Vector3fToVectorCast,							TEXT("WideningCast")}		},
-		{ CastPair{EPinType::Vector3fArray,				EPinType::VectorArray},				StatementNamePair{KCST_Vector3fToVectorArrayCast,						TEXT("WideningArrayCast")}	},
-		{ CastPair{EPinType::Vector3fSet,				EPinType::VectorSet},				StatementNamePair{KCST_Vector3fToVectorSetCast,							TEXT("WideningSetCast")}	},
-
-		{ CastPair{EPinType::Vector,					EPinType::Vector3f},				StatementNamePair{KCST_VectorToVector3fCast,							TEXT("NarrowingCast")}		},
-		{ CastPair{EPinType::VectorArray,				EPinType::Vector3fArray},			StatementNamePair{KCST_VectorToVector3fArrayCast,						TEXT("NarrowingArrayCast")}	},
-		{ CastPair{EPinType::VectorSet,					EPinType::Vector3fSet},				StatementNamePair{KCST_VectorToVector3fSetCast,							TEXT("NarrowingSetCast")}	},
-
-		{ CastPair{EPinType::FloatKeyOtherValueMap,		EPinType::DoubleKeyOtherValueMap},	StatementNamePair{KCST_FloatToDoubleKeys_MapCast,						TEXT("MapCast")}	},
-		{ CastPair{EPinType::DoubleKeyOtherValueMap,	EPinType::FloatKeyOtherValueMap},	StatementNamePair{KCST_DoubleToFloatKeys_MapCast,						TEXT("MapCast")}	},
-		{ CastPair{EPinType::OtherKeyFloatValueMap,		EPinType::OtherKeyDoubleValueMap},	StatementNamePair{KCST_FloatToDoubleValues_MapCast,						TEXT("MapCast")}	},
-		{ CastPair{EPinType::OtherKeyDoubleValueMap,	EPinType::OtherKeyFloatValueMap},	StatementNamePair{KCST_DoubleToFloatValues_MapCast,						TEXT("MapCast")}	},
-
-		{ CastPair{EPinType::FloatKeyFloatValueMap,		EPinType::DoubleKeyDoubleValueMap},	StatementNamePair{KCST_FloatToDoubleKeys_FloatToDoubleValues_MapCast,	TEXT("MapCast")}	},
-		{ CastPair{EPinType::FloatKeyFloatValueMap,		EPinType::FloatKeyDoubleValueMap},	StatementNamePair{KCST_FloatToDoubleValues_MapCast,						TEXT("MapCast")}	},
-		{ CastPair{EPinType::FloatKeyFloatValueMap,		EPinType::DoubleKeyFloatValueMap},	StatementNamePair{KCST_FloatToDoubleKeys_MapCast,						TEXT("MapCast")}	},
-
-		{ CastPair{EPinType::DoubleKeyFloatValueMap,	EPinType::DoubleKeyDoubleValueMap},	StatementNamePair{KCST_FloatToDoubleValues_MapCast,						TEXT("MapCast")}	},
-		{ CastPair{EPinType::DoubleKeyFloatValueMap,	EPinType::FloatKeyDoubleValueMap},	StatementNamePair{KCST_DoubleToFloatKeys_FloatToDoubleValues_MapCast,	TEXT("MapCast")}	},
-		{ CastPair{EPinType::DoubleKeyFloatValueMap,	EPinType::FloatKeyFloatValueMap},	StatementNamePair{KCST_DoubleToFloatKeys_MapCast,						TEXT("MapCast")}	},
-
-		{ CastPair{EPinType::DoubleKeyDoubleValueMap,	EPinType::DoubleKeyFloatValueMap},	StatementNamePair{KCST_DoubleToFloatValues_MapCast,						TEXT("MapCast")}	},
-		{ CastPair{EPinType::DoubleKeyDoubleValueMap,	EPinType::FloatKeyDoubleValueMap},	StatementNamePair{KCST_DoubleToFloatKeys_MapCast,						TEXT("MapCast")}	},
-		{ CastPair{EPinType::DoubleKeyDoubleValueMap,	EPinType::FloatKeyFloatValueMap},	StatementNamePair{KCST_DoubleToFloatKeys_DoubleToFloatValues_MapCast,	TEXT("MapCast")}	},
-
-		{ CastPair{EPinType::FloatKeyDoubleValueMap,	EPinType::DoubleKeyFloatValueMap},	StatementNamePair{KCST_FloatToDoubleKeys_DoubleToFloatValues_MapCast,	TEXT("MapCast")}	},
-		{ CastPair{EPinType::FloatKeyDoubleValueMap,	EPinType::DoubleKeyDoubleValueMap},	StatementNamePair{KCST_FloatToDoubleKeys_MapCast,						TEXT("MapCast")}	},
-		{ CastPair{EPinType::FloatKeyDoubleValueMap,	EPinType::FloatKeyFloatValueMap},	StatementNamePair{KCST_DoubleToFloatValues_MapCast,						TEXT("MapCast")}	},
-	};
+	UFunction* SetConversionFunction = BlueprintTypeConversionsClass->FindFunctionByName(TEXT("ConvertSetType"));
+	check(SetConversionFunction);
 
-	CastPair LookupPair{ GetPinType(SourcePin), GetPinType(DestinationPin) };
+	UFunction* MapConversionFunction = BlueprintTypeConversionsClass->FindFunctionByName(TEXT("ConvertMapType"));
+	check(MapConversionFunction);
 
-	const StatementNamePair* TableLookupResult = ImplicitCastTable.Find(LookupPair);
+	if (SourcePin.PinType.IsMap() && DestinationPin.PinType.IsMap())
+	{
+		if ((SourcePin.PinType.PinValueType.TerminalCategory == UEdGraphSchema_K2::PC_Real) && (DestinationPin.PinType.PinValueType.TerminalCategory == UEdGraphSchema_K2::PC_Real))
+		{
+			if ((SourcePin.PinType.PinValueType.TerminalSubCategory == UEdGraphSchema_K2::PC_Float) && ((DestinationPin.PinType.PinValueType.TerminalSubCategory == UEdGraphSchema_K2::PC_Double)))
+			{
+				return {FloatingPointCastType::Container, MapConversionFunction};
+			}
+			else if ((SourcePin.PinType.PinValueType.TerminalSubCategory == UEdGraphSchema_K2::PC_Double) && ((DestinationPin.PinType.PinValueType.TerminalSubCategory == UEdGraphSchema_K2::PC_Float)))
+			{
+				return {FloatingPointCastType::Container, MapConversionFunction};
+			}
+		}
+		else if ((SourcePin.PinType.PinValueType.TerminalCategory == UEdGraphSchema_K2::PC_Struct) && (DestinationPin.PinType.PinValueType.TerminalCategory == UEdGraphSchema_K2::PC_Struct))
+		{
+			UScriptStruct* SourceStruct = CastChecked<UScriptStruct>(SourcePin.PinType.PinValueType.TerminalSubCategoryObject.Get());
+			UScriptStruct* DestinationStruct = CastChecked<UScriptStruct>(DestinationPin.PinType.PinValueType.TerminalSubCategoryObject.Get());
 
-	TOptional<StatementNamePair> FinalResult;
+			if (FStructConversionTable::Get().GetConversionFunction(SourceStruct, DestinationStruct).IsSet())
+			{
+				return {FloatingPointCastType::Container, MapConversionFunction};
+			}
+		}
+	}
 
-	if (TableLookupResult)
+	if ((SourcePin.PinType.PinCategory == UEdGraphSchema_K2::PC_Real) && (DestinationPin.PinType.PinCategory == UEdGraphSchema_K2::PC_Real))
+	{
+		if ((SourcePin.PinType.PinSubCategory == UEdGraphSchema_K2::PC_Float) && ((DestinationPin.PinType.PinSubCategory == UEdGraphSchema_K2::PC_Double)))
+		{
+			if (SourcePin.PinType.IsArray())
+			{
+				return {FloatingPointCastType::Container, ArrayConversionFunction};
+			}
+			else if (SourcePin.PinType.IsSet())
+			{
+				return {FloatingPointCastType::Container, SetConversionFunction};
+			}
+			else if (SourcePin.PinType.IsMap())
+			{
+				return {FloatingPointCastType::Container, MapConversionFunction};
+			}
+			else
+			{
+				return {FloatingPointCastType::FloatToDouble, nullptr};
+			}
+		}
+		else if ((SourcePin.PinType.PinSubCategory == UEdGraphSchema_K2::PC_Double) && ((DestinationPin.PinType.PinSubCategory == UEdGraphSchema_K2::PC_Float)))
+		{
+			if (SourcePin.PinType.IsArray())
+			{
+				return {FloatingPointCastType::Container, ArrayConversionFunction};
+			}
+			else if (SourcePin.PinType.IsSet())
+			{
+				return {FloatingPointCastType::Container, SetConversionFunction};
+			}
+			else if (SourcePin.PinType.IsMap())
+			{
+				return {FloatingPointCastType::Container, MapConversionFunction};
+			}
+			else
+			{
+				return {FloatingPointCastType::DoubleToFloat, nullptr};
+			}
+		}
+	}
+	else if ((SourcePin.PinType.PinCategory == UEdGraphSchema_K2::PC_Struct) && (DestinationPin.PinType.PinCategory == UEdGraphSchema_K2::PC_Struct))
 	{
-		FinalResult = *TableLookupResult;
+		UScriptStruct* SourceStruct = CastChecked<UScriptStruct>(SourcePin.PinType.PinSubCategoryObject.Get());
+		UScriptStruct* DestinationStruct = CastChecked<UScriptStruct>(DestinationPin.PinType.PinSubCategoryObject.Get());
+
+		if (TOptional<ConversionFunctionPairT> ConversionPair = FStructConversionTable::Get().GetConversionFunction(SourceStruct, DestinationStruct))
+		{
+			if (SourcePin.PinType.IsArray())
+			{
+				return {FloatingPointCastType::Container, ArrayConversionFunction};
+			}
+			else if (SourcePin.PinType.IsSet())
+			{
+				return {FloatingPointCastType::Container, SetConversionFunction};
+			}
+			else if (SourcePin.PinType.IsMap())
+			{
+				return {FloatingPointCastType::Container, MapConversionFunction};
+			}
+			else
+			{
+				return {FloatingPointCastType::Struct, ConversionPair->Get<1>()};
+			}
+		}
 	}
 
-	return FinalResult;
+	return {FloatingPointCastType::None, nullptr};
 }
 
 } // namespace UE::KismetCompiler::CastingUtils
-
diff --git a/Engine/Source/Editor/KismetCompiler/Private/KismetCompilerMisc.cpp b/Engine/Source/Editor/KismetCompiler/Private/KismetCompilerMisc.cpp
--- a/Engine/Source/Editor/KismetCompiler/Private/KismetCompilerMisc.cpp
+++ b/Engine/Source/Editor/KismetCompiler/Private/KismetCompilerMisc.cpp
@@ -932,10 +932,7 @@ void FKismetCompilerUtilities::CreateObjectAssignmentStatement(FKismetFunctionCo
 		// Some terms don't necessarily have a valid SourcePin (eg: FKCHandler_FunctionEntry)
 		if (DstPinSearchKey)
 		{
-			TOptional<TPair<FBPTerminal*, EKismetCompiledStatementType>> ImplicitCastEntry =
-				CastingUtils::InsertImplicitCastStatement(Context, DstPinSearchKey, RHSTerm);
-
-			ImplicitCastTerm = ImplicitCastEntry ? ImplicitCastEntry->Get<0>() : nullptr;
+			ImplicitCastTerm = CastingUtils::InsertImplicitCastStatement(Context, DstPinSearchKey, RHSTerm);
 		}
 
 		if (ImplicitCastTerm != nullptr)
@@ -2043,25 +2040,7 @@ bool FKismetCompilerUtilities::CheckFunctionCompiledStatementsThreadSafety(const
 			case KCST_CastObjToInterface:
 			case KCST_DynamicCast:
 			case KCST_DoubleToFloatCast:
-			case KCST_DoubleToFloatArrayCast:
-			case KCST_DoubleToFloatSetCast:
 			case KCST_FloatToDoubleCast:
-			case KCST_FloatToDoubleArrayCast:
-			case KCST_FloatToDoubleSetCast:
-			case KCST_VectorToVector3fCast:
-			case KCST_VectorToVector3fArrayCast:
-			case KCST_VectorToVector3fSetCast:
-			case KCST_Vector3fToVectorCast:
-			case KCST_Vector3fToVectorArrayCast:
-			case KCST_Vector3fToVectorSetCast:
-			case KCST_FloatToDoubleKeys_MapCast:
-			case KCST_DoubleToFloatKeys_MapCast:
-			case KCST_FloatToDoubleValues_MapCast:
-			case KCST_DoubleToFloatValues_MapCast:
-			case KCST_FloatToDoubleKeys_FloatToDoubleValues_MapCast:
-			case KCST_DoubleToFloatKeys_FloatToDoubleValues_MapCast:
-			case KCST_DoubleToFloatKeys_DoubleToFloatValues_MapCast:
-			case KCST_FloatToDoubleKeys_DoubleToFloatValues_MapCast:
 			case KCST_ObjectToBool:
 				break;
 			case KCST_AddMulticastDelegate:
diff --git a/Engine/Source/Editor/KismetCompiler/Private/KismetCompilerVMBackend.cpp b/Engine/Source/Editor/KismetCompiler/Private/KismetCompilerVMBackend.cpp
--- a/Engine/Source/Editor/KismetCompiler/Private/KismetCompilerVMBackend.cpp
+++ b/Engine/Source/Editor/KismetCompiler/Private/KismetCompilerVMBackend.cpp
@@ -2072,51 +2072,9 @@ public:
 			case KCST_DoubleToFloatCast:
 				CastType = CST_DoubleToFloat;
 				break;
-			case KCST_DoubleToFloatArrayCast:
-				CastType = CST_DoubleToFloatArray;
-				break;
-			case KCST_DoubleToFloatSetCast:
-				CastType = CST_DoubleToFloatSet;
-				break;
 			case KCST_FloatToDoubleCast:
 				CastType = CST_FloatToDouble;
 				break;
-			case KCST_FloatToDoubleArrayCast:
-				CastType = CST_FloatToDoubleArray;
-				break;
-			case KCST_FloatToDoubleSetCast:
-				CastType = CST_FloatToDoubleSet;
-				break;
-			case KCST_VectorToVector3fCast:
-				CastType = CST_VectorToVector3f;
-				break;
-			case KCST_Vector3fToVectorCast:
-				CastType = CST_Vector3fToVector;
-				break;
-			case KCST_FloatToDoubleKeys_MapCast:
-				CastType = CST_FloatToDoubleKeys_Map;
-				break;
-			case KCST_DoubleToFloatKeys_MapCast:
-				CastType = CST_DoubleToFloatKeys_Map;
-				break;
-			case KCST_FloatToDoubleValues_MapCast:
-				CastType = CST_FloatToDoubleValues_Map;
-				break;
-			case KCST_DoubleToFloatValues_MapCast:
-				CastType = CST_DoubleToFloatValues_Map;
-				break;
-			case KCST_FloatToDoubleKeys_FloatToDoubleValues_MapCast:
-				CastType = CST_FloatToDoubleKeys_FloatToDoubleValues_Map;
-				break;
-			case KCST_DoubleToFloatKeys_FloatToDoubleValues_MapCast:
-				CastType = CST_DoubleToFloatKeys_FloatToDoubleValues_Map;
-				break;
-			case KCST_DoubleToFloatKeys_DoubleToFloatValues_MapCast:
-				CastType = CST_DoubleToFloatKeys_DoubleToFloatValues_Map;
-				break;
-			case KCST_FloatToDoubleKeys_DoubleToFloatValues_MapCast:
-				CastType = CST_FloatToDoubleKeys_DoubleToFloatValues_Map;
-				break;
 			default:
 				check(false);
 				break;
@@ -2254,25 +2212,7 @@ public:
 			EmitCreateMapStatement(Statement);
 			break;
 		case KCST_DoubleToFloatCast:
-		case KCST_DoubleToFloatArrayCast:
-		case KCST_DoubleToFloatSetCast:
 		case KCST_FloatToDoubleCast:
-		case KCST_FloatToDoubleArrayCast:
-		case KCST_FloatToDoubleSetCast:
-		case KCST_VectorToVector3fCast:
-		case KCST_VectorToVector3fArrayCast:
-		case KCST_VectorToVector3fSetCast:
-		case KCST_Vector3fToVectorCast:
-		case KCST_Vector3fToVectorArrayCast:
-		case KCST_Vector3fToVectorSetCast:
-		case KCST_FloatToDoubleKeys_MapCast:
-		case KCST_DoubleToFloatKeys_MapCast:
-		case KCST_FloatToDoubleValues_MapCast:
-		case KCST_DoubleToFloatValues_MapCast:
-		case KCST_FloatToDoubleKeys_FloatToDoubleValues_MapCast:
-		case KCST_DoubleToFloatKeys_FloatToDoubleValues_MapCast:
-		case KCST_DoubleToFloatKeys_DoubleToFloatValues_MapCast:
-		case KCST_FloatToDoubleKeys_DoubleToFloatValues_MapCast:
 			EmitCastStatement(Statement);
 			break;
 		default:
diff --git a/Engine/Source/Editor/KismetCompiler/Public/BlueprintCompiledStatement.h b/Engine/Source/Editor/KismetCompiler/Public/BlueprintCompiledStatement.h
--- a/Engine/Source/Editor/KismetCompiler/Public/BlueprintCompiledStatement.h
+++ b/Engine/Source/Editor/KismetCompiler/Public/BlueprintCompiledStatement.h
@@ -68,30 +68,7 @@ enum EKismetCompiledStatementType
 	KCST_SwitchValue = 29,
 	
 	KCST_DoubleToFloatCast = 30,
-	KCST_DoubleToFloatArrayCast = 31,
-	KCST_DoubleToFloatSetCast = 32,
-	
-	KCST_FloatToDoubleCast = 33,
-	KCST_FloatToDoubleArrayCast = 34,
-	KCST_FloatToDoubleSetCast = 35,
-	
-	KCST_VectorToVector3fCast = 36,
-	KCST_VectorToVector3fArrayCast = 37,
-	KCST_VectorToVector3fSetCast = 38,
-	
-	KCST_Vector3fToVectorCast = 39,
-	KCST_Vector3fToVectorArrayCast = 40,
-	KCST_Vector3fToVectorSetCast = 41,
-
-	KCST_FloatToDoubleKeys_MapCast = 42,
-	KCST_DoubleToFloatKeys_MapCast = 43,
-	KCST_FloatToDoubleValues_MapCast = 44,
-	KCST_DoubleToFloatValues_MapCast = 45,
-
-	KCST_FloatToDoubleKeys_FloatToDoubleValues_MapCast = 46,
-	KCST_DoubleToFloatKeys_FloatToDoubleValues_MapCast = 47,
-	KCST_DoubleToFloatKeys_DoubleToFloatValues_MapCast = 48,
-	KCST_FloatToDoubleKeys_DoubleToFloatValues_MapCast = 49,
+	KCST_FloatToDoubleCast = 31,
 
 	//~ Kismet instrumentation extensions:
 
diff --git a/Engine/Source/Editor/KismetCompiler/Public/KismetCastingUtils.h b/Engine/Source/Editor/KismetCompiler/Public/KismetCastingUtils.h
--- a/Engine/Source/Editor/KismetCompiler/Public/KismetCastingUtils.h
+++ b/Engine/Source/Editor/KismetCompiler/Public/KismetCastingUtils.h
@@ -11,17 +11,40 @@ class UEdGraphPin;
 
 namespace UE::KismetCompiler::CastingUtils
 {
-	using StatementNamePair = TPair<EKismetCompiledStatementType, const TCHAR*>;
+	enum class FloatingPointCastType
+	{
+		None,
+		FloatToDouble,
+		DoubleToFloat,
+		Container,
+		Struct
+	};
+
+	struct FConversion
+	{
+		FloatingPointCastType Type = FloatingPointCastType::None;
+		UFunction* Function = nullptr;
+	};
+
+	struct FImplicitCastParams
+	{
+		FConversion Conversion;
+		FBPTerminal* TargetTerminal = nullptr;
+		UEdGraphNode* TargetNode = nullptr;
+	};
 
 	/**
-	 * Given a specific EKismetCompiledStatementType enum, this function returns its inverse.
-	 * eg: the inverse of KCST_DoubleToFloatCast is KCST_FloatToDoubleCast.
-	 * Invalid enums (ie: not cast-related) will return an unset TOptional.
+	 * Directly creates a new FBPTerminal that holds the result of an implicit cast.
+	 * Internally, this is used by InsertImplicitCastStatement, but it can be used by any node
+	 * that has special semantics for handling implicit casts. The terminal is guaranteed to have 
+	 * a unique name in the context.
 	 * 
-	 * @param Statement - The enum to query for an inverse.
-	 * @return A set TOptional with the enum's inverse. Unset if no inverse is found.
+	 * @param Context - Current function context to add the terminal to.
+	 * @param Net - The corresponding pin that we'll copy type details from.
+	 * @param SourceNode - The source node for the terminal. This is usually the Net's owner, but it can be different if necessary.
+	 * @return A new FBPTerminal for the temporary variable that holds the casted result.
 	 */
-	KISMETCOMPILER_API TOptional<EKismetCompiledStatementType> GetInverseCastStatement(EKismetCompiledStatementType Statement);
+	KISMETCOMPILER_API FBPTerminal* MakeImplicitCastTerminal(FKismetFunctionContext& Context, UEdGraphPin* Net, UEdGraphNode* SourceNode = nullptr);
 
 	/**
 	 * Analyzes the NetMap of the current function context for potential implicit casts.
@@ -35,17 +58,29 @@ namespace UE::KismetCompiler::CastingUtils
 
 	/**
 	 * Utility function used by nodes for inserting implicit cast statements.
-	 * During compilation, a node that potentially may need to handle a cast should call this function.
-	 * If the current pin needs a cast, a statement is inserted, and a new terminal for the temporary is returned.
+	 * A compiled cast statement is created on the target node with the given parameters.
 	 * 
 	 * @param Context - Current function context to analyze. Assumes the ImplicitCastMap has been populated.
+	 * @param CastParams - Parameters for the cast operation.
+	 * @param RHSTerm - The current terminal that should have its data read from.
+	 */
+	KISMETCOMPILER_API void InsertImplicitCastStatement(FKismetFunctionContext& Context,
+														const FImplicitCastParams& CastParams,
+														FBPTerminal* RHSTerm);
+
+	/**
+	 * Utility function used by nodes for inserting implicit cast statements.
+	 * Similar to above, but the cast parameters are read from the context's cast map using
+	 * a destination pin as the key. If no entry exists in the map, null is returned.
+	 *
+	 * @param Context - Current function context to analyze. Assumes the ImplicitCastMap has been populated.
 	 * @param DestinationPin - Used as a key in the ImplicitCastMap. These pins are always inputs.
 	 * @param RHSTerm - The current terminal that should have its data read from.
-	 * @return A new FBPTerminal, EKismetCompiledStatementType pair for the temporary variable that has the casted result (if one exists)
+	 * @return A new FBPTerminal for the temporary variable that holds the casted result (if one exists).
 	 */
-	KISMETCOMPILER_API TOptional<TPair<FBPTerminal*, EKismetCompiledStatementType>> InsertImplicitCastStatement(FKismetFunctionContext& Context,
-																												UEdGraphPin* DestinationPin,
-																												FBPTerminal* RHSTerm);
+	KISMETCOMPILER_API FBPTerminal* InsertImplicitCastStatement(FKismetFunctionContext& Context,
+																UEdGraphPin* DestinationPin,
+																FBPTerminal* RHSTerm);
 
 	/**
 	 * Removes the specific UEdGraphPin from the context's implicit cast map.
@@ -61,14 +96,14 @@ namespace UE::KismetCompiler::CastingUtils
 	/**
 	 * Retrieves the conversion type needed between two arbitrary pins (if necessary). Specifically, this indicates if either
 	 * a narrowing or widening cast is needed between a float or a double type (including containers). In addition to the
-	 * corresponding EKismetCompiledStatementType that represents the cast type, a string literal describing the cast is also
-	 * returned.
+	 * corresponding FloatingPointCastType that represents the cast type, a UFunction* may be included for container and
+	 * struct conversions.
 	 * 
 	 * @param SourcePin - The source pin to compare.
 	 * @param DestinationPin - The destination pin to compare.
-	 * @return A new StatementNamePair containing the cast information. The result is unset if no conversion is needed.
+	 * @return A new FConversion containing the cast information.
 	 */
-	KISMETCOMPILER_API TOptional<StatementNamePair> GetFloatingPointConversionType(const UEdGraphPin& SourcePin, const UEdGraphPin& DestinationPin);
+	KISMETCOMPILER_API FConversion GetFloatingPointConversion(const UEdGraphPin& SourcePin, const UEdGraphPin& DestinationPin);
 
 } // UE::KismetCompiler::CastingUtils
 
diff --git a/Engine/Source/Editor/KismetCompiler/Public/KismetCompiledFunctionContext.h b/Engine/Source/Editor/KismetCompiler/Public/KismetCompiledFunctionContext.h
--- a/Engine/Source/Editor/KismetCompiler/Public/KismetCompiledFunctionContext.h
+++ b/Engine/Source/Editor/KismetCompiler/Public/KismetCompiledFunctionContext.h
@@ -9,6 +9,7 @@
 #include "BPTerminal.h"
 #include "BlueprintCompiledStatement.h"
 #include "Kismet2/CompilerResultsLog.h"
+#include "KismetCastingUtils.h"
 
 class Error;
 class UBlueprint;
@@ -32,13 +33,6 @@ enum ETerminalSpecification
 	TS_ForcedShared,
 };
 
-struct FImplicitCastParams
-{
-	EKismetCompiledStatementType CastType = KCST_Nop;
-	FBPTerminal* TargetTerminal = nullptr;
-	UEdGraphNode* TargetNode = nullptr;
-};
-
 //////////////////////////////////////////////////////////////////////////
 // FKismetFunctionContext
 
@@ -92,7 +86,7 @@ public:
 	TMap<UEdGraphPin*, FBPTerminal*> LiteralHackMap;
 
 	// Contains a map of destination pins that will need an implicit cast to either a float or double
-	TMap<UEdGraphPin*, FImplicitCastParams> ImplicitCastMap;
+	TMap<UEdGraphPin*, UE::KismetCompiler::CastingUtils::FImplicitCastParams> ImplicitCastMap;
 
 	bool bIsUbergraph;
 	bool bCannotBeCalledFromOtherKismet;
diff --git a/Engine/Source/Runtime/CoreUObject/Private/UObject/ScriptCore.cpp b/Engine/Source/Runtime/CoreUObject/Private/UObject/ScriptCore.cpp
--- a/Engine/Source/Runtime/CoreUObject/Private/UObject/ScriptCore.cpp
+++ b/Engine/Source/Runtime/CoreUObject/Private/UObject/ScriptCore.cpp
@@ -24,7 +24,6 @@
 #include "UObject/SoftObjectPtr.h"
 #include "UObject/PropertyPortFlags.h"
 #include "UObject/UnrealType.h"
-#include "UObject/ScriptCastingUtils.h"
 #include "UObject/Stack.h"
 #include "UObject/Reload.h"
 #include "Blueprint/BlueprintSupport.h"
@@ -3873,7 +3872,10 @@ DEFINE_FUNCTION(UObject::execDoubleToFloatCast)
 {
 	if (Stack.StepAndCheckMostRecentProperty(Stack.Object, nullptr))
 	{
-		DoubleToFloatCast(nullptr, Stack.MostRecentPropertyAddress, RESULT_PARAM);
+		const double* Source = reinterpret_cast<const double*>(Stack.MostRecentPropertyAddress);
+		float* Destination = reinterpret_cast<float*>(RESULT_PARAM);
+
+		*Destination = static_cast<float>(*Source);
 	}
 	else
 	{
@@ -3882,116 +3884,21 @@ DEFINE_FUNCTION(UObject::execDoubleToFloatCast)
 }
 IMPLEMENT_CAST_FUNCTION( CST_DoubleToFloat, execDoubleToFloatCast )
 
-DEFINE_FUNCTION(UObject::execDoubleToFloatArrayCast)
-{
-	CopyAndCastArrayFromStack<double, float>(Stack, RESULT_PARAM);
-}
-IMPLEMENT_CAST_FUNCTION( CST_DoubleToFloatArray, execDoubleToFloatArrayCast )
-
-DEFINE_FUNCTION(UObject::execDoubleToFloatSetCast)
-{
-	CopyAndCastSetFromStack<double, float>(Stack, RESULT_PARAM);
-}
-IMPLEMENT_CAST_FUNCTION( CST_DoubleToFloatSet, execDoubleToFloatSetCast )
-
 DEFINE_FUNCTION(UObject::execFloatToDoubleCast)
 {
 	if (Stack.StepAndCheckMostRecentProperty(Stack.Object, nullptr))
 	{
-		FloatToDoubleCast(nullptr, Stack.MostRecentPropertyAddress, RESULT_PARAM);
-	}
-	else
-	{
-		UE_LOG(LogScript, Verbose, TEXT("Cast failed: recent properties were null!"));
-	}
-}
-IMPLEMENT_CAST_FUNCTION( CST_FloatToDouble, execFloatToDoubleCast )
-
-DEFINE_FUNCTION(UObject::execFloatToDoubleArrayCast)
-{
-	CopyAndCastArrayFromStack<float, double>(Stack, RESULT_PARAM);
-}
-IMPLEMENT_CAST_FUNCTION( CST_FloatToDoubleArray, execFloatToDoubleArrayCast )
-
-DEFINE_FUNCTION(UObject::execFloatToDoubleSetCast)
-{
-	CopyAndCastSetFromStack<float, double>(Stack, RESULT_PARAM);
-}
-IMPLEMENT_CAST_FUNCTION( CST_FloatToDoubleSet, execFloatToDoubleSetCast )
-
-DEFINE_FUNCTION(UObject::execVectorToVector3fCast)
-{
-	if (Stack.StepAndCheckMostRecentProperty(Stack.Object, nullptr))
-	{
-		FloatingPointCast<FVector, FVector3f>(nullptr, Stack.MostRecentPropertyAddress, RESULT_PARAM);
-	}
-	else
-	{
-		UE_LOG(LogScript, Verbose, TEXT("Cast failed: recent properties were null!"));
-	}
-}
-IMPLEMENT_CAST_FUNCTION( CST_VectorToVector3f, execVectorToVector3fCast )
+		const float* Source = reinterpret_cast<const float*>(Stack.MostRecentPropertyAddress);
+		double* Destination = reinterpret_cast<double*>(RESULT_PARAM);
 
-DEFINE_FUNCTION(UObject::execVector3fToVectorCast)
-{
-	if (Stack.StepAndCheckMostRecentProperty(Stack.Object, nullptr))
-	{
-		FloatingPointCast<FVector3f, FVector>(nullptr, Stack.MostRecentPropertyAddress, RESULT_PARAM);
+		*Destination = *Source;
 	}
 	else
 	{
 		UE_LOG(LogScript, Verbose, TEXT("Cast failed: recent properties were null!"));
 	}
 }
-IMPLEMENT_CAST_FUNCTION( CST_Vector3fToVector, execVector3fToVectorCast )
-
-DEFINE_FUNCTION(UObject::execFloatToDoubleKeysMapCast)
-{
-	CopyAndCastMapFromStack<FloatToDoubleCast, CopyElement>(Stack, RESULT_PARAM);
-}
-IMPLEMENT_CAST_FUNCTION( CST_FloatToDoubleKeys_Map, execFloatToDoubleKeysMapCast )
-
-DEFINE_FUNCTION(UObject::execDoubleToFloatKeysMapCast)
-{
-	CopyAndCastMapFromStack<DoubleToFloatCast, CopyElement>(Stack, RESULT_PARAM);
-}
-IMPLEMENT_CAST_FUNCTION( CST_DoubleToFloatKeys_Map, execDoubleToFloatKeysMapCast )
-
-DEFINE_FUNCTION(UObject::execFloatToDoubleValuesMapCast)
-{
-	CopyAndCastMapFromStack<CopyElement, FloatToDoubleCast>(Stack, RESULT_PARAM);
-}
-IMPLEMENT_CAST_FUNCTION( CST_FloatToDoubleValues_Map, execFloatToDoubleValuesMapCast )
-
-DEFINE_FUNCTION(UObject::execDoubleToFloatValuesMapCast)
-{
-	CopyAndCastMapFromStack<CopyElement, DoubleToFloatCast>(Stack, RESULT_PARAM);
-}
-IMPLEMENT_CAST_FUNCTION( CST_DoubleToFloatValues_Map, execDoubleToFloatValuesMapCast )
-
-DEFINE_FUNCTION(UObject::execFloatToDoubleKeysFloatToDoubleValuesMapCast)
-{
-	CopyAndCastMapFromStack<FloatToDoubleCast, FloatToDoubleCast>(Stack, RESULT_PARAM);
-}
-IMPLEMENT_CAST_FUNCTION( CST_FloatToDoubleKeys_FloatToDoubleValues_Map, execFloatToDoubleKeysFloatToDoubleValuesMapCast )
-
-DEFINE_FUNCTION(UObject::execDoubleToFloatKeysFloatToDoubleValuesMapCast)
-{
-	CopyAndCastMapFromStack<DoubleToFloatCast, FloatToDoubleCast>(Stack, RESULT_PARAM);
-}
-IMPLEMENT_CAST_FUNCTION( CST_DoubleToFloatKeys_FloatToDoubleValues_Map, execDoubleToFloatKeysFloatToDoubleValuesMapCast )
-
-DEFINE_FUNCTION(UObject::execDoubleToFloatKeysDoubleToFloatValuesMapCast)
-{
-	CopyAndCastMapFromStack<DoubleToFloatCast, DoubleToFloatCast>(Stack, RESULT_PARAM);
-}
-IMPLEMENT_CAST_FUNCTION( CST_DoubleToFloatKeys_DoubleToFloatValues_Map, execDoubleToFloatKeysDoubleToFloatValuesMapCast )
-
-DEFINE_FUNCTION(UObject::execFloatToDoubleKeysDoubleToFloatValuesMapCast)
-{
-	CopyAndCastMapFromStack<FloatToDoubleCast, DoubleToFloatCast>(Stack, RESULT_PARAM);
-}
-IMPLEMENT_CAST_FUNCTION( CST_FloatToDoubleKeys_DoubleToFloatValues_Map, execFloatToDoubleKeysDoubleToFloatValuesMapCast )
+IMPLEMENT_CAST_FUNCTION( CST_FloatToDouble, execFloatToDoubleCast )
 
 DEFINE_FUNCTION(UObject::execObjectToBool)
 {
diff --git a/Engine/Source/Runtime/CoreUObject/Public/UObject/Object.h b/Engine/Source/Runtime/CoreUObject/Public/UObject/Object.h
--- a/Engine/Source/Runtime/CoreUObject/Public/UObject/Object.h
+++ b/Engine/Source/Runtime/CoreUObject/Public/UObject/Object.h
@@ -1569,30 +1569,13 @@ public:
 	DECLARE_FUNCTION(execMetaCast);
 	DECLARE_FUNCTION(execInterfaceCast);
 	DECLARE_FUNCTION(execDoubleToFloatCast);
-	DECLARE_FUNCTION(execDoubleToFloatArrayCast);
-	DECLARE_FUNCTION(execDoubleToFloatSetCast);
 	DECLARE_FUNCTION(execFloatToDoubleCast);
-	DECLARE_FUNCTION(execFloatToDoubleArrayCast);
-	DECLARE_FUNCTION(execFloatToDoubleSetCast);
-	DECLARE_FUNCTION(execVectorToVector3fCast);
-	DECLARE_FUNCTION(execVector3fToVectorCast);
 	DECLARE_FUNCTION(execObjectToBool);
 	DECLARE_FUNCTION(execInterfaceToBool);
 	DECLARE_FUNCTION(execObjectToInterface);
 	DECLARE_FUNCTION(execInterfaceToInterface);
 	DECLARE_FUNCTION(execInterfaceToObject);
 
-	// Map-specific conversions
-	DECLARE_FUNCTION(execFloatToDoubleKeysMapCast);
-	DECLARE_FUNCTION(execDoubleToFloatKeysMapCast);
-	DECLARE_FUNCTION(execFloatToDoubleValuesMapCast);
-	DECLARE_FUNCTION(execDoubleToFloatValuesMapCast);
-	DECLARE_FUNCTION(execFloatToDoubleKeysFloatToDoubleValuesMapCast);
-	DECLARE_FUNCTION(execDoubleToFloatKeysFloatToDoubleValuesMapCast);
-	DECLARE_FUNCTION(execDoubleToFloatKeysDoubleToFloatValuesMapCast);
-	DECLARE_FUNCTION(execFloatToDoubleKeysDoubleToFloatValuesMapCast);
-
-
 	// Dynamic array functions
 	// Array support
 	DECLARE_FUNCTION(execGetDynArrayElement);
diff --git a/Engine/Source/Runtime/CoreUObject/Public/UObject/Script.h b/Engine/Source/Runtime/CoreUObject/Public/UObject/Script.h
--- a/Engine/Source/Runtime/CoreUObject/Public/UObject/Script.h
+++ b/Engine/Source/Runtime/CoreUObject/Public/UObject/Script.h
@@ -303,21 +303,7 @@ enum ECastToken
 	CST_ObjectToBool			= 0x01,
 	CST_InterfaceToBool			= 0x02,
 	CST_DoubleToFloat			= 0x03,
-	CST_DoubleToFloatArray		= 0x04,
-	CST_DoubleToFloatSet		= 0x05,
-	CST_FloatToDouble			= 0x06,
-	CST_FloatToDoubleArray		= 0x07,
-	CST_FloatToDoubleSet		= 0x08,
-	CST_VectorToVector3f		= 0x09,
-	CST_Vector3fToVector		= 0x0A,
-	CST_FloatToDoubleKeys_Map	= 0x0B,
-	CST_DoubleToFloatKeys_Map	= 0x0C,
-	CST_FloatToDoubleValues_Map	= 0x0D,
-	CST_DoubleToFloatValues_Map	= 0x0E,
-	CST_FloatToDoubleKeys_FloatToDoubleValues_Map	= 0x0F,
-	CST_DoubleToFloatKeys_FloatToDoubleValues_Map	= 0x10,
-	CST_DoubleToFloatKeys_DoubleToFloatValues_Map	= 0x11,
-	CST_FloatToDoubleKeys_DoubleToFloatValues_Map	= 0x12,
+	CST_FloatToDouble			= 0x04,
 
 	CST_Max						= 0xFF,
 };
diff --git a/Engine/Source/Runtime/CoreUObject/Public/UObject/ScriptCastingUtils.h b/Engine/Source/Runtime/CoreUObject/Public/UObject/ScriptCastingUtils.h
deleted file mode 100644
diff --git a/Engine/Source/Runtime/Engine/Internal/Kismet/BlueprintTypeConversions.h b/Engine/Source/Runtime/Engine/Internal/Kismet/BlueprintTypeConversions.h
new file mode 100644
--- /dev/null
+++ b/Engine/Source/Runtime/Engine/Internal/Kismet/BlueprintTypeConversions.h
@@ -0,0 +1,74 @@
+// Copyright Epic Games, Inc. All Rights Reserved.
+
+#pragma once
+
+#include "CoreMinimal.h"
+
+#include "BlueprintTypeConversions.generated.h"
+
+namespace UE::Kismet::BlueprintTypeConversions
+{
+	using ConversionFunctionT = void(*)(const void*, void*);
+	using ConversionFunctionPairT = TPair<ConversionFunctionT, UFunction*>;
+
+	class ENGINE_API FStructConversionTable
+	{
+	public:
+		static const FStructConversionTable& Get();
+
+		TOptional<UE::Kismet::BlueprintTypeConversions::ConversionFunctionPairT>
+			GetConversionFunction(const UScriptStruct* InFrom, const UScriptStruct* InTo) const;
+
+	private:
+		FStructConversionTable();
+
+		using StructVariantPairT = TPair<const UScriptStruct*, const UScriptStruct*>;
+		using ImplicitCastLookupTableT = TMap<StructVariantPairT, UE::Kismet::BlueprintTypeConversions::ConversionFunctionPairT>;
+		using StructVariantLookupTableT = TMap<const UScriptStruct*, const UScriptStruct*>;
+
+		StructVariantPairT GetVariantsFromStructs(const UScriptStruct* InStruct1, const UScriptStruct* InStruct2) const;
+
+		ImplicitCastLookupTableT ImplicitCastLookupTable;
+		StructVariantLookupTableT StructVariantLookupTable;
+
+		static FStructConversionTable* Instance;
+	};
+}
+
+#define MAKE_CONVERSION_EXEC_FUNCTION_NAME(SourceType, DestType)						\
+	execConvert##SourceType##To##DestType
+
+#define DECLARE_CONVERSION_FUNCTIONS(BaseType, VariantType1, VariantType2)				\
+	DECLARE_FUNCTION(MAKE_CONVERSION_EXEC_FUNCTION_NAME(VariantType1, VariantType2));	\
+	DECLARE_FUNCTION(MAKE_CONVERSION_EXEC_FUNCTION_NAME(VariantType2, VariantType1));
+
+UCLASS()
+class ENGINE_API UBlueprintTypeConversions : public UObject
+{
+public:
+	GENERATED_UCLASS_BODY()
+
+	// Container conversions
+
+	UFUNCTION(CustomThunk, meta = (BlueprintInternalUseOnly = "true"))
+	static TArray<int> ConvertArrayType(const TArray<int>& InArray);
+	DECLARE_FUNCTION(execConvertArrayType);
+
+	UFUNCTION(CustomThunk, meta = (BlueprintInternalUseOnly = "true"))
+	static TSet<int> ConvertSetType(const TSet<int>& InSet);
+	DECLARE_FUNCTION(execConvertSetType);
+
+	UFUNCTION(CustomThunk, meta = (BlueprintInternalUseOnly = "true"))
+	static TMap<int, int> ConvertMapType(const TMap<int, int>& InMap);
+	DECLARE_FUNCTION(execConvertMapType);
+
+	// Custom struct conversions
+
+	DECLARE_CONVERSION_FUNCTIONS(FVector, FVector3f, FVector3d);
+
+	UFUNCTION(CustomThunk, meta = (BlueprintInternalUseOnly = "true"))
+	static int32 ConvertFVector3dToFVector3f(int32 InFromData);
+
+	UFUNCTION(CustomThunk, meta = (BlueprintInternalUseOnly = "true"))
+	static int32 ConvertFVector3fToFVector3d(int32 InFromData);
+};
diff --git a/Engine/Source/Runtime/Engine/Private/Kismet/BlueprintTypeConversions.cpp b/Engine/Source/Runtime/Engine/Private/Kismet/BlueprintTypeConversions.cpp
new file mode 100644
--- /dev/null
+++ b/Engine/Source/Runtime/Engine/Private/Kismet/BlueprintTypeConversions.cpp
@@ -0,0 +1,373 @@
+// Copyright Epic Games, Inc. All Rights Reserved.
+
+#include "Kismet/BlueprintTypeConversions.h"
+
+#define MAKE_CONVERSION_FUNCTION_NAME(SourceType, DestType)													\
+	Convert##SourceType##To##DestType
+
+#define DEFINE_CONVERSION_FUNCTIONS(BASETYPE, VARIANTTYPE1, VARIANTTYPE2, IMPL1TO2, IMPL2TO1)				\
+	DEFINE_FUNCTION(UBlueprintTypeConversions::MAKE_CONVERSION_EXEC_FUNCTION_NAME(VARIANTTYPE1, VARIANTTYPE2))	\
+	{																										\
+		void* DestAddr = Stack.MostRecentPropertyAddress;													\
+		Stack.MostRecentProperty = nullptr;																	\
+		Stack.StepCompiledIn(nullptr, nullptr);																\
+		const void* SourceAddr = Stack.MostRecentPropertyAddress;											\
+		P_FINISH;																							\
+		IMPL1TO2(SourceAddr, DestAddr);																		\
+	}																										\
+	DEFINE_FUNCTION(UBlueprintTypeConversions::MAKE_CONVERSION_EXEC_FUNCTION_NAME(VARIANTTYPE2, VARIANTTYPE1))	\
+	{																										\
+		void* DestAddr = Stack.MostRecentPropertyAddress;													\
+		Stack.MostRecentProperty = nullptr;																	\
+		Stack.StepCompiledIn(nullptr, nullptr);																\
+		const void* SourceAddr = Stack.MostRecentPropertyAddress;											\
+		P_FINISH;																							\
+		IMPL2TO1(SourceAddr, DestAddr);																		\
+	}																										\
+	UE::Kismet::BlueprintTypeConversions::Internal::FStructConversionEntry BASETYPE##Entry(					\
+		&TBaseStructure<BASETYPE>::Get,																		\
+		&TVariantStructure<BASETYPE>::Get,																	\
+		&TVariantStructure<VARIANTTYPE1>::Get,																\
+		&TVariantStructure<VARIANTTYPE2>::Get,																\
+		TEXT(PREPROCESSOR_TO_STRING(MAKE_CONVERSION_FUNCTION_NAME(VARIANTTYPE1, VARIANTTYPE2))),			\
+		TEXT(PREPROCESSOR_TO_STRING(MAKE_CONVERSION_FUNCTION_NAME(VARIANTTYPE2, VARIANTTYPE1))),			\
+		&IMPL1TO2,																							\
+		&IMPL2TO1																							\
+	);
+
+namespace UE::Kismet::BlueprintTypeConversions::Internal
+{
+
+ConversionFunctionT FindConversionFunction(const FProperty* InFromProperty, const FProperty* InToProperty)
+{
+	ConversionFunctionT Result = nullptr;
+
+	auto ConvertFloatToDouble = [](const void* InFromData, void* InToData)
+	{
+		const float* FromFloat = reinterpret_cast<const float*>(InFromData);
+		check(FromFloat);
+		double* ToDouble = reinterpret_cast<double*>(InToData);
+		check(ToDouble);
+
+		*ToDouble = *FromFloat;
+	};
+
+	auto ConvertDoubleToFloat = [](const void* InFromData, void* InToData)
+	{
+		const double* FromDouble = reinterpret_cast<const double*>(InFromData);
+		check(FromDouble);
+		float* ToFloat = reinterpret_cast<float*>(InToData);
+		check(ToFloat);
+
+		*ToFloat = static_cast<float>(*FromDouble);
+	};
+
+	if (InFromProperty->IsA<FFloatProperty>() && InToProperty->IsA<FDoubleProperty>())
+	{
+		Result = ConvertFloatToDouble;
+	}
+	else if (InFromProperty->IsA<FDoubleProperty>() && InToProperty->IsA<FFloatProperty>())
+	{
+		Result = ConvertDoubleToFloat;
+	}
+	else if (InFromProperty->IsA<FStructProperty>() && InToProperty->IsA<FStructProperty>())
+	{
+		const FStructProperty* InFromStructProperty = CastFieldChecked<FStructProperty>(InFromProperty);
+		const FStructProperty* InToStructProperty = CastFieldChecked<FStructProperty>(InToProperty);
+		const UScriptStruct* InFromStruct = InFromStructProperty->Struct;
+		const UScriptStruct* InToStruct = InToStructProperty->Struct;
+
+		TOptional<ConversionFunctionPairT> ConversionPair = FStructConversionTable::Get().GetConversionFunction(InFromStruct, InToStruct);
+		check(ConversionPair.IsSet());
+		Result = ConversionPair->Get<0>();
+	}
+
+	return Result;
+}
+
+struct FStructConversionEntry
+{
+	using GetUScriptStructFunctionT = UScriptStruct * (*)(void);
+
+	FStructConversionEntry(GetUScriptStructFunctionT InGetBaseStruct,
+		GetUScriptStructFunctionT InGetVariantFromBaseStruct,
+		GetUScriptStructFunctionT InGetVariantStruct1,
+		GetUScriptStructFunctionT InGetVariantStruct2,
+		const TCHAR* InConvertVariant1ToVariant2FunctionName,
+		const TCHAR* InConvertVariant2ToVariant1FunctionName,
+		UE::Kismet::BlueprintTypeConversions::ConversionFunctionT InConvertVariant1ToVariant2Impl,
+		UE::Kismet::BlueprintTypeConversions::ConversionFunctionT InConvertVariant2ToVariant1Impl)
+		: GetBaseStruct(InGetBaseStruct)
+		, GetVariantFromBaseStruct(InGetVariantFromBaseStruct)
+		, GetVariantStruct1(InGetVariantStruct1)
+		, GetVariantStruct2(InGetVariantStruct2)
+		, ConvertVariant1ToVariant2FunctionName(InConvertVariant1ToVariant2FunctionName)
+		, ConvertVariant2ToVariant1FunctionName(InConvertVariant2ToVariant1FunctionName)
+		, ConvertVariant1ToVariant2Impl(InConvertVariant1ToVariant2Impl)
+		, ConvertVariant2ToVariant1Impl(InConvertVariant2ToVariant1Impl)
+	{
+		NextEntry = FStructConversionEntry::GetListHead();
+		FStructConversionEntry::GetListHead() = this;
+	}
+
+	GetUScriptStructFunctionT GetBaseStruct;
+	GetUScriptStructFunctionT GetVariantFromBaseStruct;
+	GetUScriptStructFunctionT GetVariantStruct1;
+	GetUScriptStructFunctionT GetVariantStruct2;
+	const TCHAR* ConvertVariant1ToVariant2FunctionName;
+	const TCHAR* ConvertVariant2ToVariant1FunctionName;
+	UE::Kismet::BlueprintTypeConversions::ConversionFunctionT ConvertVariant1ToVariant2Impl;
+	UE::Kismet::BlueprintTypeConversions::ConversionFunctionT ConvertVariant2ToVariant1Impl;
+	FStructConversionEntry* NextEntry;
+
+	static FStructConversionEntry*& GetListHead()
+	{
+		static FStructConversionEntry* ListHead = nullptr;
+		return ListHead;
+	}
+};
+
+} // namespace UE::Kismet::BlueprintTypeConversions::Internal
+
+namespace UE::Kismet::BlueprintTypeConversions
+{
+
+FStructConversionTable* FStructConversionTable::Instance = nullptr;
+
+FStructConversionTable::FStructConversionTable()
+{
+	using namespace UE::Kismet::BlueprintTypeConversions;
+
+	UClass* BlueprintTypeConversionsClass = UBlueprintTypeConversions::StaticClass();
+	check(BlueprintTypeConversionsClass);
+
+	Internal::FStructConversionEntry* Entry = Internal::FStructConversionEntry::GetListHead();
+	while (Entry)
+	{
+		StructVariantLookupTable.Add(Entry->GetBaseStruct(), Entry->GetVariantFromBaseStruct());
+
+		StructVariantPairT ConversionKeyPair1 = { Entry->GetVariantStruct1(), Entry->GetVariantStruct2() };
+		ConversionFunctionPairT ConversionValuePair1 = 
+			{ Entry->ConvertVariant1ToVariant2Impl,
+			  BlueprintTypeConversionsClass->FindFunctionByName(Entry->ConvertVariant1ToVariant2FunctionName) };
+
+		StructVariantPairT ConversionKeyPair2 = { Entry->GetVariantStruct2(), Entry->GetVariantStruct1() };
+		ConversionFunctionPairT ConversionValuePair2 = 
+			{ Entry->ConvertVariant2ToVariant1Impl,
+			  BlueprintTypeConversionsClass->FindFunctionByName(Entry->ConvertVariant2ToVariant1FunctionName) };
+
+		ImplicitCastLookupTable.Add(ConversionKeyPair1, ConversionValuePair1);
+		ImplicitCastLookupTable.Add(ConversionKeyPair2, ConversionValuePair2);
+
+		Entry = Entry->NextEntry;
+	}
+}
+
+const FStructConversionTable& FStructConversionTable::Get()
+{
+	if (Instance == nullptr)
+	{
+		Instance = new FStructConversionTable();
+	}
+
+	return *Instance;
+}
+
+TOptional<ConversionFunctionPairT> FStructConversionTable::GetConversionFunction(const UScriptStruct* InFrom, const UScriptStruct* InTo) const
+{
+	TOptional<ConversionFunctionPairT> Result;
+
+	StructVariantPairT Key = GetVariantsFromStructs(InFrom, InTo);
+
+	if (const ConversionFunctionPairT* Entry = ImplicitCastLookupTable.Find(Key))
+	{
+		Result = *Entry;
+	}
+
+	return Result;
+}
+
+FStructConversionTable::StructVariantPairT FStructConversionTable::GetVariantsFromStructs(const UScriptStruct* InStruct1, const UScriptStruct* InStruct2) const
+{
+	StructVariantPairT Result;
+
+	if (const UScriptStruct* const* Variant = StructVariantLookupTable.Find(InStruct1))
+	{
+		Result.Get<0>() = *Variant;
+	}
+	else
+	{
+		Result.Get<0>() = InStruct1;
+	}
+
+	if (const UScriptStruct* const* Variant = StructVariantLookupTable.Find(InStruct2))
+	{
+		Result.Get<1>() = *Variant;
+	}
+	else
+	{
+		Result.Get<1>() = InStruct2;
+	}
+
+	return Result;
+}
+
+} // namespace UE::Kismet::BlueprintTypeConversions
+
+UBlueprintTypeConversions::UBlueprintTypeConversions(const FObjectInitializer& ObjectInitializer)
+	: Super(ObjectInitializer)
+{
+}
+
+DEFINE_FUNCTION(UBlueprintTypeConversions::execConvertArrayType)
+{
+	using namespace UE::Kismet;
+
+	FArrayProperty* DestArrayProperty = CastFieldChecked<FArrayProperty>(Stack.MostRecentProperty);
+	void* DestArrayAddr = Stack.MostRecentPropertyAddress;
+
+	Stack.MostRecentProperty = nullptr;
+	Stack.StepCompiledIn<FArrayProperty>(nullptr);
+	const void* SourceArrayAddr = Stack.MostRecentPropertyAddress;
+	const FArrayProperty* SourceArrayProperty = CastFieldChecked<FArrayProperty>(Stack.MostRecentProperty);
+
+	P_FINISH;
+
+	FScriptArrayHelper SourceArray(SourceArrayProperty, SourceArrayAddr);
+	FScriptArrayHelper DestArray(DestArrayProperty, DestArrayAddr);
+
+	int SourceArraySize = SourceArray.Num();
+	DestArray.Resize(SourceArraySize);
+
+	BlueprintTypeConversions::ConversionFunctionT ConversionFunction =
+		BlueprintTypeConversions::Internal::FindConversionFunction(SourceArrayProperty->Inner, DestArrayProperty->Inner);
+	check(ConversionFunction);
+
+	for (int i = 0; i < SourceArraySize; ++i)
+	{
+		const void* SrcData = SourceArray.GetRawPtr(i);
+		void* DestData = DestArray.GetRawPtr(i);
+		(*ConversionFunction)(SrcData, DestData);
+	}
+}
+
+DEFINE_FUNCTION(UBlueprintTypeConversions::execConvertSetType)
+{
+	using namespace UE::Kismet;
+
+	FSetProperty* DestSetProperty = CastFieldChecked<FSetProperty>(Stack.MostRecentProperty);
+	void* DestSetAddr = Stack.MostRecentPropertyAddress;
+
+	Stack.MostRecentProperty = nullptr;
+	Stack.StepCompiledIn<FSetProperty>(nullptr);
+	const void* SourceSetAddr = Stack.MostRecentPropertyAddress;
+	const FSetProperty* SourceSetProperty = CastFieldChecked<FSetProperty>(Stack.MostRecentProperty);
+
+	P_FINISH;
+
+	FScriptSetHelper SourceSet(SourceSetProperty, SourceSetAddr);
+	FScriptSetHelper DestSet(DestSetProperty, DestSetAddr);
+
+	int SourceSetSize = SourceSet.Num();
+	DestSet.EmptyElements(SourceSetSize);
+
+	BlueprintTypeConversions::ConversionFunctionT ConversionFunction = 
+		BlueprintTypeConversions::Internal::FindConversionFunction(SourceSetProperty->ElementProp, DestSetProperty->ElementProp);
+	check(ConversionFunction);
+
+	for (int i = 0; i < SourceSetSize; ++i)
+	{
+		const void* SrcData = SourceSet.GetElementPtr(i);
+
+		int32 NewIndex = DestSet.AddDefaultValue_Invalid_NeedsRehash();
+		void* DestData = DestSet.GetElementPtr(NewIndex);
+		(*ConversionFunction)(SrcData, DestData);
+	}
+
+	DestSet.Rehash();
+}
+
+DEFINE_FUNCTION(UBlueprintTypeConversions::execConvertMapType)
+{
+	using namespace UE::Kismet;
+
+	FMapProperty* DestMapProperty = CastFieldChecked<FMapProperty>(Stack.MostRecentProperty);
+	void* DestMapAddr = Stack.MostRecentPropertyAddress;
+
+	Stack.MostRecentProperty = nullptr;
+	Stack.StepCompiledIn<FMapProperty>(nullptr);
+	const void* SourceMapAddr = Stack.MostRecentPropertyAddress;
+	const FMapProperty* SourceMapProperty = CastFieldChecked<FMapProperty>(Stack.MostRecentProperty);
+
+	P_FINISH;
+
+	FScriptMapHelper SourceMap(SourceMapProperty, SourceMapAddr);
+	FScriptMapHelper DestMap(DestMapProperty, DestMapAddr);
+
+	int SourceMapSize = SourceMap.Num();
+	DestMap.EmptyValues(SourceMapSize);
+
+	BlueprintTypeConversions::ConversionFunctionT KeyConversionFunction = 
+		BlueprintTypeConversions::Internal::FindConversionFunction(SourceMapProperty->KeyProp, DestMapProperty->KeyProp);
+	BlueprintTypeConversions::ConversionFunctionT ValueConversionFunction = 
+		BlueprintTypeConversions::Internal::FindConversionFunction(SourceMapProperty->ValueProp, DestMapProperty->ValueProp);
+
+	for (int i = 0; i < SourceMapSize; ++i)
+	{
+		int32 NewIndex = DestMap.AddDefaultValue_Invalid_NeedsRehash();
+
+		const void* SourceKeyRawData = SourceMap.GetKeyPtr(i);
+		void* DestinationKeyRawData = DestMap.GetKeyPtr(NewIndex);
+
+		if (KeyConversionFunction)
+		{
+			(*KeyConversionFunction)(SourceKeyRawData, DestinationKeyRawData);
+		}
+		else
+		{
+			SourceMapProperty->KeyProp->CopySingleValue(DestinationKeyRawData, SourceKeyRawData);
+		}
+
+		const void* SourceValueRawData = SourceMap.GetValuePtr(i);
+		void* DestinationValueRawData = DestMap.GetValuePtr(NewIndex);
+
+		if (ValueConversionFunction)
+		{
+			(*ValueConversionFunction)(SourceValueRawData, DestinationValueRawData);
+		}
+		else
+		{
+			SourceMapProperty->ValueProp->CopySingleValue(DestinationValueRawData, SourceValueRawData);
+		}
+	}
+
+	DestMap.Rehash();
+}
+
+// Container conversions
+
+FORCEINLINE_DEBUGGABLE void ConvertFVector3dToFVector3fImpl(const void* InFromData, void* InToData)
+{
+	const FVector3d* From = reinterpret_cast<const FVector3d*>(InFromData);
+	check(From);
+	FVector3f* To = reinterpret_cast<FVector3f*>(InToData);
+	check(To);
+
+	float x = static_cast<float>(From->X);
+	float y = static_cast<float>(From->Y);
+	float z = static_cast<float>(From->Z);
+
+	*To = FVector3f(x, y, z);
+}
+
+FORCEINLINE_DEBUGGABLE void ConvertFVector3fToFVector3dImpl(const void* InFromData, void* InToData)
+{
+	const FVector3f* From = reinterpret_cast<const FVector3f*>(InFromData);
+	check(From);
+	FVector3d* To = reinterpret_cast<FVector3d*>(InToData);
+	check(To);
+
+	*To = FVector3d(From->X, From->Y, From->Z);
+}
+
+DEFINE_CONVERSION_FUNCTIONS(FVector, FVector3d, FVector3f, ConvertFVector3dToFVector3fImpl, ConvertFVector3fToFVector3dImpl)
diff --git a/Engine/Source/Runtime/Engine/Private/PropertyAccess.cpp b/Engine/Source/Runtime/Engine/Private/PropertyAccess.cpp
--- a/Engine/Source/Runtime/Engine/Private/PropertyAccess.cpp
+++ b/Engine/Source/Runtime/Engine/Private/PropertyAccess.cpp
@@ -2,7 +2,6 @@
 
 #include "PropertyAccess.h"
 #include "Misc/MemStack.h"
-#include "UObject/ScriptCastingUtils.h"
 
 #define LOCTEXT_NAMESPACE "PropertyAccess"
 
@@ -385,7 +384,7 @@ struct FPropertyAccessSystem
 			checkSlow(InDestProperty->IsA<FArrayProperty>());
 			const FArrayProperty* SrcArrayProperty = ExactCastField<const FArrayProperty>(InSrcProperty);
 			const FArrayProperty* DestArrayProperty = ExactCastField<const FArrayProperty>(InDestProperty);
-			CopyAndCastArray<float, double>(SrcArrayProperty, InSrcAddr, DestArrayProperty, InDestAddr);
+			CopyAndCastFloatingPointArray<float, double>(SrcArrayProperty, InSrcAddr, DestArrayProperty, InDestAddr);
 			break;
 		}
 		case EPropertyAccessCopyType::DemoteArrayDoubleToFloat:
@@ -394,7 +393,7 @@ struct FPropertyAccessSystem
 			checkSlow(InDestProperty->IsA<FArrayProperty>());
 			const FArrayProperty* SrcArrayProperty = ExactCastField<const FArrayProperty>(InSrcProperty);
 			const FArrayProperty* DestArrayProperty = ExactCastField<const FArrayProperty>(InDestProperty);
-			CopyAndCastArray<double, float>(SrcArrayProperty, InSrcAddr, DestArrayProperty, InDestAddr);
+			CopyAndCastFloatingPointArray<double, float>(SrcArrayProperty, InSrcAddr, DestArrayProperty, InDestAddr);
 			break;
 		}
 		default:
@@ -606,6 +605,30 @@ struct FPropertyAccessSystem
 			}
 		}
 	}
+
+	template <typename SourceType, typename DestinationType>
+	static void CopyAndCastFloatingPointArray(const FArrayProperty* SourceArrayProperty,
+											  const void* SourceAddress,
+											  const FArrayProperty* DestinationArrayProperty,
+											  void* DestinationAddress)
+	{
+		checkSlow(SourceArrayProperty);
+		checkSlow(SourceAddress);
+		checkSlow(DestinationArrayProperty);
+		checkSlow(DestinationAddress);
+
+		FScriptArrayHelper SourceArrayHelper(SourceArrayProperty, SourceAddress);
+		FScriptArrayHelper DestinationArrayHelper(DestinationArrayProperty, DestinationAddress);
+
+		DestinationArrayHelper.Resize(SourceArrayHelper.Num());
+		for (int32 i = 0; i < SourceArrayHelper.Num(); ++i)
+		{
+			const SourceType* SourceData = reinterpret_cast<const SourceType*>(SourceArrayHelper.GetRawPtr(i));
+			DestinationType* DestinationData = reinterpret_cast<DestinationType*>(DestinationArrayHelper.GetRawPtr(i));
+
+			*DestinationData = static_cast<DestinationType>(*SourceData);
+		}
+	}
 };
 
 const FPropertyAccessLibrary& FPropertyAccessLibrary::operator =(const FPropertyAccessLibrary& Other)
