Non-unity build fix. Missing forward decl.

diff --git a/Engine/Source/Editor/KismetCompiler/Public/KismetCastingUtils.h b/Engine/Source/Editor/KismetCompiler/Public/KismetCastingUtils.h
--- a/Engine/Source/Editor/KismetCompiler/Public/KismetCastingUtils.h
+++ b/Engine/Source/Editor/KismetCompiler/Public/KismetCastingUtils.h
@@ -7,6 +7,7 @@
 
 struct FBPTerminal;
 struct FKismetFunctionContext;
+class UEdGraphNode;
 class UEdGraphPin;
 
 namespace UE::KismetCompiler::CastingUtils
