[BugFix] Changing a variable type inside of a Structure doesn't reset it's value to the default value causing "Cannot parse value" error

diff --git a/Engine/Source/Editor/KismetCompiler/Private/UserDefinedStructureCompilerUtils.cpp b/Engine/Source/Editor/KismetCompiler/Private/UserDefinedStructureCompilerUtils.cpp
--- a/Engine/Source/Editor/KismetCompiler/Private/UserDefinedStructureCompilerUtils.cpp
+++ b/Engine/Source/Editor/KismetCompiler/Private/UserDefinedStructureCompilerUtils.cpp
@@ -150,9 +150,6 @@ struct FUserDefinedStructureCompilerInner
 
 		if (UUserDefinedStructEditorData* EditorData = Cast<UUserDefinedStructEditorData>(StructToClean->EditorData))
 		{
-			// Ensure that editor data is in sync w/ the current default instance (if valid) so that it can be reinitialized later.
-			EditorData->RefreshValuesFromDefaultInstance();
-
 			EditorData->CleanDefaultInstance();
 		}
 
@@ -391,6 +388,12 @@ struct FUserDefinedStructureCompilerInner
 
 			FUserDefinedStructureCompilerInner::CleanAndSanitizeStruct(Struct);
 			FUserDefinedStructureCompilerInner::InnerCompileStruct(Struct, GetDefault<UEdGraphSchema_K2>(), MessageLog);
+			
+			if (UUserDefinedStructEditorData* EditorData = Cast<UUserDefinedStructEditorData>(Struct->EditorData))
+			{
+				// Ensure that editor data is in sync w/ the current default instance (if valid) so that it can be reinitialized later.
+				EditorData->RefreshValuesFromDefaultInstance();
+			}
 
 			DependencyMap.RemoveAtSwap(StructureToCompileIndex);
 
