UE-156275 - Infinite Loop due to Macros setup does not get detected, causing the Editor to freeze indefinitely on compile

Prior to tunnel&macro expansion, we need to add a step that checks for any cycles with macros. In this case, we simply perform a DFS of macro graph instances. If a cycle is detected, the BP compiler will print the cycle the path and issue an error. Failure to do so will cause the function to continuously add nodes to the source graph as macro instances are discovered, which results in an infinite loop.

diff --git a/Engine/Source/Editor/KismetCompiler/Private/KismetCompiler.cpp b/Engine/Source/Editor/KismetCompiler/Private/KismetCompiler.cpp
--- a/Engine/Source/Editor/KismetCompiler/Private/KismetCompiler.cpp
+++ b/Engine/Source/Editor/KismetCompiler/Private/KismetCompiler.cpp
@@ -457,6 +457,67 @@ namespace UE::KismetCompiler::Private
 			}
 		}
 	}
+
+	bool FindMacroCycle(const UK2Node_MacroInstance* RootNode, TSet<FGuid>& VisitedMacroGraphs, TArray<const UK2Node_MacroInstance*>& CurrentPath)
+	{
+		check(RootNode);
+		if (UEdGraph* MacroGraph = RootNode->GetMacroGraph())
+		{
+			VisitedMacroGraphs.Add(MacroGraph->GraphGuid);
+			CurrentPath.Push(RootNode);
+			for (const UEdGraphNode* ChildNode : MacroGraph->Nodes)
+			{
+				if (const UK2Node_MacroInstance* MacroInstanceNode = Cast<const UK2Node_MacroInstance>(ChildNode))
+				{
+					UEdGraph* InnerMacroGraph = MacroInstanceNode->GetMacroGraph();
+					if (VisitedMacroGraphs.Contains(InnerMacroGraph->GraphGuid))
+					{
+						return true;
+					}
+					else
+					{
+						return FindMacroCycle(MacroInstanceNode, VisitedMacroGraphs, CurrentPath);
+					}
+				}
+			}
+			CurrentPath.Pop();
+		}
+
+		return false;
+	}
+
+	bool FindAndReportMacroCycles(const UEdGraph* SourceGraph, FCompilerResultsLog& MessageLog)
+	{
+		check(SourceGraph);
+
+		TArray<const UK2Node_MacroInstance*> RootLevelMacroNodes;
+		for (const UEdGraphNode* Node : SourceGraph->Nodes)
+		{
+			if (const UK2Node_MacroInstance* MacroInstanceNode = Cast<const UK2Node_MacroInstance>(Node))
+			{
+				RootLevelMacroNodes.Add(MacroInstanceNode);
+			}
+		}
+
+		for (const UK2Node_MacroInstance* RootNode : RootLevelMacroNodes)
+		{
+			TSet<FGuid> VisitedMacroGraphs;
+			TArray<const UK2Node_MacroInstance*> CurrentPath;
+			if (FindMacroCycle(RootNode, VisitedMacroGraphs, CurrentPath))
+			{
+				MessageLog.Error(*LOCTEXT("MacroCycleDetected_ErrorFmt", "Macro cycle detected in @@! Cycle path:").ToString(), RootNode);
+
+				for (const UK2Node_MacroInstance* Node : CurrentPath)
+				{
+					MessageLog.Error(*LOCTEXT("MacroCyclePath_ErrorFmt", "\t@@").ToString(), Node);
+				}
+
+				return true;
+			}
+		}
+
+		return false;
+	}
 }
 
 //////////////////////////////////////////////////////////////////////////
@@ -3990,6 +4051,14 @@ void FKismetCompilerContext::CreateCommentBlockAroundNodes(const TArray<UEdGraph
 
 void FKismetCompilerContext::ExpandTunnelsAndMacros(UEdGraph* SourceGraph)
 {
+	check(SourceGraph);
+
+	// Macro cycles can lead to infinite creation of nodes and hang the editor.
+	if (UE::KismetCompiler::Private::FindAndReportMacroCycles(SourceGraph, MessageLog))
+	{
+		return;
+	}
+
 	// Determine if we are regenerating a blueprint on load
 	const bool bIsLoading = Blueprint ? Blueprint->bIsRegeneratingOnLoad : false;
 
@@ -4015,7 +4084,7 @@ void FKismetCompilerContext::ExpandTunnelsAndMacros(UEdGraph* SourceGraph)
 		{
 			UEdGraph* MacroGraph = MacroInstanceNode->GetMacroGraph();
 			// Verify that this macro can actually be expanded
-			if( MacroGraph == NULL )
+			if (MacroGraph == nullptr)
 			{
 				MessageLog.Error(TEXT("Macro node @@ is pointing at an invalid macro graph."), MacroInstanceNode);
 				continue;
@@ -4030,7 +4099,7 @@ void FKismetCompilerContext::ExpandTunnelsAndMacros(UEdGraph* SourceGraph)
 			const bool bForceRegenNodes = bIsLoading && MacroBlueprint && (MacroBlueprint != Blueprint) && !MacroBlueprint->bHasBeenRegenerated;
 
 			// Clone the macro graph, then move all of its children, keeping a list of nodes from the macro
-			UEdGraph* ClonedGraph = FEdGraphUtilities::CloneGraph(MacroGraph, NULL, &MessageLog, true);
+			UEdGraph* ClonedGraph = FEdGraphUtilities::CloneGraph(MacroGraph, nullptr, &MessageLog, true);
 
 			for (int32 I = 0; I < ClonedGraph->Nodes.Num(); ++I)
 			{
@@ -4139,7 +4208,7 @@ void FKismetCompilerContext::ExpandTunnelsAndMacros(UEdGraph* SourceGraph)
 			{
 				UEdGraphNode* DuplicatedNode = *MacroNodeIt;
 				
-				if( DuplicatedNode != NULL )
+				if (DuplicatedNode != nullptr)
 				{
 					if (bForceRegenNodes)
 					{
