PR #9347: Fix: Wrong documentation of FNodeHandlingFunctor (Contributed by [REDACTED_PERSON])

diff --git a/Engine/Source/Editor/KismetCompiler/Public/KismetCompilerMisc.h b/Engine/Source/Editor/KismetCompiler/Public/KismetCompilerMisc.h
--- a/Engine/Source/Editor/KismetCompiler/Public/KismetCompilerMisc.h
+++ b/Engine/Source/Editor/KismetCompiler/Public/KismetCompilerMisc.h
@@ -153,8 +153,7 @@ public:
 	virtual ~FNodeHandlingFunctor() 
 	{
 	}
-
-	//virtual void Validate(FKismetFunctionContext& Context, UEdGraphNode* Node) {}
+	
 	virtual void Compile(FKismetFunctionContext& Context, UEdGraphNode* Node)
 	{
 	}
