Fix it so all blueprint functions are marked with RF_LoadCompleted after compilation so they won't get corrupted during partial reload of packages
This was already happening for properties and ubergraph functions to fix similar bugs with async loading

diff --git a/Engine/Source/Editor/KismetCompiler/Private/KismetCompiler.cpp b/Engine/Source/Editor/KismetCompiler/Private/KismetCompiler.cpp
--- a/Engine/Source/Editor/KismetCompiler/Private/KismetCompiler.cpp
+++ b/Engine/Source/Editor/KismetCompiler/Private/KismetCompiler.cpp
@@ -933,11 +933,8 @@ FProperty* FKismetCompilerContext::CreateVariable(const FName VarName, const FEd
 	FProperty* NewProperty = FKismetCompilerUtilities::CreatePropertyOnScope(NewClass, VarName, VarType, NewClass, CPF_None, Schema, MessageLog);
 	if (NewProperty != nullptr)
 	{
-		// This fixes a rare bug involving asynchronous loading of BPs in editor builds. The pattern was established
-		// in FKismetCompilerContext::CompileFunctions where we do this for the uber graph function. By setting
-		// the RF_LoadCompleted we prevent the linker from overwriting our regenerated property, although the
-		// circumstances under which this occurs are murky. More testing of BPs loading asynchronously in the editor
-		// needs to be added:
+		// This is necessary to stop async loading and partial reloads from corrupting the compiled data
+		// without marking the owning class for recompilation
 		NewProperty->SetFlags(RF_LoadCompleted);
 		FKismetCompilerUtilities::LinkAddedProperty(NewClass, NewProperty);
 	}
@@ -2786,6 +2783,13 @@ void FKismetCompilerContext::FinishCompilingFunction(FKismetFunctionContext& Con
 	{
 		FKismetCompilerUtilities::CheckFunctionThreadSafety(Context, MessageLog);
 	}
+
+	// This is necessary to stop async loading and partial reloads from corrupting the compiled data
+	// without marking the owning class for recompilation
+	if (Context.Function)
+	{
+		Context.Function->SetFlags(RF_LoadCompleted);
+	}
 }
 
 void FKismetCompilerContext::SetCalculatedMetaDataAndFlags(UFunction* Function, UK2Node_FunctionEntry* EntryNode, const UEdGraphSchema_K2* K2Schema)
@@ -4827,12 +4831,6 @@ void FKismetCompilerContext::CompileFunctions(EInternalCompilerFlags InternalFla
 		}
 	}
 
-	// It's necessary to tell if UberGraphFunction is ready to create frame.
-	if (NewClass->UberGraphFunction)
-	{
-		NewClass->UberGraphFunction->SetFlags(RF_LoadCompleted);
-	}
-
 	{ BP_SCOPED_COMPILER_EVENT_STAT(EKismetCompilerStats_FinalizationWork);
 
 		// Set any final flags and seal the class, build a CDO, etc...
