UE-162610 - Error: Assertion failed: IntFitsIn<OutType>(In) at 'FScriptBuilderBase::EmitTermExpr'

Reverted the IntCastChecked for setting byte constants. Apparently, we have a case where this gets tripped, which needs to be followed up separately.

diff --git a/Engine/Source/Editor/KismetCompiler/Private/KismetCompilerVMBackend.cpp b/Engine/Source/Editor/KismetCompiler/Private/KismetCompilerVMBackend.cpp
--- a/Engine/Source/Editor/KismetCompiler/Private/KismetCompilerVMBackend.cpp
+++ b/Engine/Source/Editor/KismetCompiler/Private/KismetCompilerVMBackend.cpp
@@ -800,11 +800,11 @@ public:
 				if (EnumPtr)
 				{
 					//Get value from enum string
-					Value = IntCastChecked<uint8, int64>(EnumPtr->GetValueByName(*(Term->Name)));
+					Value = EnumPtr->GetValueByName(*(Term->Name));
 				}
 				else
 				{
-					Value = IntCastChecked<uint8, int32>(FCString::Atoi(*(Term->Name)));
+					Value = FCString::Atoi(*(Term->Name));
 				}
 
 				Writer << EX_ByteConst;
