Added support for specifying the name of a cloned graph.

This is primarily useful for debugging cloned function graphs for the intermediate build products. Without an explicit name, it's difficult to determine where a cloned graph originated from.

diff --git a/Engine/Source/Editor/KismetCompiler/Private/KismetCompiler.cpp b/Engine/Source/Editor/KismetCompiler/Private/KismetCompiler.cpp
--- a/Engine/Source/Editor/KismetCompiler/Private/KismetCompiler.cpp
+++ b/Engine/Source/Editor/KismetCompiler/Private/KismetCompiler.cpp
@@ -4332,6 +4332,22 @@ void FKismetCompilerContext::ProcessOneFunctionGraph(UEdGraph* SourceGraph, bool
 
 	// Clone the source graph so we can modify it as needed; merging in the child graphs
 	UEdGraph* FunctionGraph = FEdGraphUtilities::CloneGraph(SourceGraph, Blueprint, &MessageLog, true);
+
+	// Preserve the original graph name, which helps with debugging intermediate products
+	{
+		FString FunctionGraphName = SourceGraph->GetName() + TEXT("_MERGED");
+
+		ERenameFlags RenameFlags =
+			(REN_NonTransactional | REN_DoNotDirty | REN_DontCreateRedirectors | REN_ForceNoResetLoaders);
+
+		if (UEdGraph* ExistingGraph = FindObject<UEdGraph>(Blueprint, *FunctionGraphName))
+		{
+			ExistingGraph->Rename(nullptr, GetTransientPackage(), RenameFlags);
+		}
+
+		FunctionGraph->Rename(*FunctionGraphName, nullptr, RenameFlags);
+	}
+
 	FEdGraphUtilities::MergeChildrenGraphsIn(FunctionGraph, FunctionGraph, /*bRequireSchemaMatch=*/ true);
 
 	ExpansionStep(FunctionGraph, false);
