UE-145634 - Delegate fixup doesn't work when bound functions are not in the current BP

Another consequence of using doubles everywhere in Blueprints is that delegate binding can fail if a native delegate signature uses floats. The initial attempt at fixing this would update modify function pins in the current Blueprint to use PC_Float as a subcategory if it was used in a delegate binding. Unfortunately, this doesn't work if we're binding to a function in a different BP class. Attempting to fix up those functions would dirty several Blueprints, and likely confuse users.

The approach that we take here creates a proxy delegate function that matches the underlying signature. The proxy's function graph will then call the original bound function, which implicitly handles any double/float conversions. Additionally, we'll perform a "capture" of any actors that we need to call the bound function on. The compiler will add a new hidden property for the actor class, which we set in the original graph that had the bound delegate. The proxy function graph will then read this property when it calls the original function.

diff --git a/Engine/Source/Editor/GraphEditor/GraphEditor.Build.cs b/Engine/Source/Editor/GraphEditor/GraphEditor.Build.cs
--- a/Engine/Source/Editor/GraphEditor/GraphEditor.Build.cs
+++ b/Engine/Source/Editor/GraphEditor/GraphEditor.Build.cs
@@ -12,12 +12,6 @@ public class GraphEditor : ModuleRules
 				"StructViewer",
 			}
         );
-         
-//         PublicDependencyModuleNames.AddRange(
-//             new string[] {
-//                 "AudioEditor"
-//             }
-//         );
 
 		PrivateDependencyModuleNames.AddRange(
 			new string[] {
@@ -34,6 +28,7 @@ public class GraphEditor : ModuleRules
 				"UnrealEd",
 				"AssetRegistry",
 				"Kismet",
+				"KismetCompiler",
 				"KismetWidgets",
 				"BlueprintGraph",
 				"Documentation",
diff --git a/Engine/Source/Editor/GraphEditor/Private/KismetNodes/SGraphNodeK2CreateDelegate.cpp b/Engine/Source/Editor/GraphEditor/Private/KismetNodes/SGraphNodeK2CreateDelegate.cpp
--- a/Engine/Source/Editor/GraphEditor/Private/KismetNodes/SGraphNodeK2CreateDelegate.cpp
+++ b/Engine/Source/Editor/GraphEditor/Private/KismetNodes/SGraphNodeK2CreateDelegate.cpp
@@ -16,6 +16,7 @@
 #include "K2Node_CustomEvent.h"
 #include "Kismet2/BlueprintEditorUtils.h"
 #include "Kismet2/KismetEditorUtilities.h"
+#include "KismetCompilerMisc.h"
 #include "K2Node.h"
 #include "K2Node_CreateDelegate.h"
 #include "Layout/Margin.h"
@@ -308,7 +309,8 @@ void SGraphNodeK2CreateDelegate::CreateBelowPinControls(TSharedPtr<SVerticalBox>
 			for (TFieldIterator<UFunction> It(ScopeClass); It; ++It)
 			{
 				UFunction* Func = *It;
-				if (Func && FunctionSignature->IsSignatureCompatibleWith(Func) &&
+				if (Func && 
+					(FKismetCompilerUtilities::DoSignaturesHaveConvertibleFloatTypes(Func, FunctionSignature) != ConvertibleSignatureMatchResult::Different) &&
 					UEdGraphSchema_K2::FunctionCanBeUsedInDelegate(Func))
 				{
 					FFunctionItemData ItemData;
diff --git a/Engine/Source/Editor/Kismet/Private/BlueprintCompilationManager.cpp b/Engine/Source/Editor/Kismet/Private/BlueprintCompilationManager.cpp
--- a/Engine/Source/Editor/Kismet/Private/BlueprintCompilationManager.cpp
+++ b/Engine/Source/Editor/Kismet/Private/BlueprintCompilationManager.cpp
@@ -19,6 +19,7 @@
 #include "FileHelpers.h"
 #include "FindInBlueprintManager.h"
 #include "IMessageLogListing.h"
+#include "K2Node_CreateDelegate.h"
 #include "K2Node_CustomEvent.h"
 #include "K2Node_FunctionEntry.h"
 #include "K2Node_FunctionResult.h"
@@ -2881,6 +2882,28 @@ UClass* FBlueprintCompilationManagerImpl::FastGenerateSkeletonClass(UBlueprint*
 		}
 	};
 
+	const auto CreateDelegateProxyFunctions = [Ret, &CompilerContext](UField**& InCurrentFieldStorageLocation)
+	{
+		for (const auto& DelegateInfo : CompilerContext.ConvertibleDelegates)
+		{
+			check(DelegateInfo.Key);
+
+			const UFunction* DelegateSignature = DelegateInfo.Key->GetDelegateSignature();
+
+			// The original signature function is actually a UDelegateFunction type.
+			// When we create our own function from the original, we need to ensure that it's a standard UFunction.
+
+			UObject* NewObject =
+				StaticDuplicateObject(DelegateSignature, Ret, DelegateInfo.Value.ProxyFunctionName, RF_AllFlags, UFunction::StaticClass());
+			UFunction* NewFunction = CastChecked<UFunction>(NewObject);
+
+			NewFunction->FunctionFlags &= ~(FUNC_Delegate | FUNC_MulticastDelegate);
+
+			*InCurrentFieldStorageLocation = NewFunction;
+			InCurrentFieldStorageLocation = &NewFunction->Next;
+		}
+	};
+
 	Ret->SetSuperStruct(ParentClass);
 	
 	Ret->ClassFlags |= (ParentClass->ClassFlags & CLASS_Inherit);
@@ -2958,12 +2981,19 @@ UClass* FBlueprintCompilationManagerImpl::FastGenerateSkeletonClass(UBlueprint*
 		}
 	}
 
+	{
+		CompilerContext.NewClass = Ret;
+		CompilerContext.RegisterClassDelegateProxiesFromBlueprint();
+		CompilerContext.NewClass = OriginalNewClass;
+	}
+
 	{
 		CompilerContext.NewClass = Ret;
 		TGuardValue<bool> GuardAssignDelegateSignatureFunction( CompilerContext.bAssignDelegateSignatureFunction, true);
 		CompilerContext.CreateClassVariablesFromBlueprint();
 		CompilerContext.NewClass = OriginalNewClass;
 	}
+
 	UField* Iter = Ret->Children;
 	while(Iter)
 	{
@@ -2971,6 +3001,8 @@ UClass* FBlueprintCompilationManagerImpl::FastGenerateSkeletonClass(UBlueprint*
 		Iter = Iter->Next;
 	}
 	
+	CreateDelegateProxyFunctions(CurrentFieldStorageLocation);
+
 	AddFunctionForGraphs(TEXT(""), BP->FunctionGraphs, CurrentFieldStorageLocation, BPTYPE_FunctionLibrary == BP->BlueprintType, false);
 	AddFunctionForGraphs(TEXT(""), CompilerContext.GeneratedFunctionGraphs, CurrentFieldStorageLocation, BPTYPE_FunctionLibrary == BP->BlueprintType, false);
 
diff --git a/Engine/Source/Editor/KismetCompiler/Private/KismetCompiler.cpp b/Engine/Source/Editor/KismetCompiler/Private/KismetCompiler.cpp
--- a/Engine/Source/Editor/KismetCompiler/Private/KismetCompiler.cpp
+++ b/Engine/Source/Editor/KismetCompiler/Private/KismetCompiler.cpp
@@ -1180,6 +1180,8 @@ void FKismetCompilerContext::CreateClassVariablesFromBlueprint()
 {
 	BP_SCOPED_COMPILER_EVENT_STAT(EKismetCompilerStats_CreateClassVariables);
 
+	check(NewClass);
+
 	// Grab the blueprint variables
 	NewClass->NumReplicatedProperties = 0;	// Keep track of how many replicated variables this blueprint adds
 
@@ -3453,6 +3455,23 @@ FName FKismetCompilerContext::GetEventStubFunctionName(UK2Node_Event* SrcEventNo
 	return EventNodeName;
 }
 
+void FKismetCompilerContext::RegisterClassDelegateProxiesFromBlueprint()
+{
+	check(Blueprint);
+
+	ConvertibleDelegates.Empty();
+
+	for (const TObjectPtr<UEdGraph>& Graph : Blueprint->FunctionGraphs)
+	{
+		RegisterConvertibleDelegates(Graph);
+	}
+
+	for (const TObjectPtr<UEdGraph>& Graph : Blueprint->UbergraphPages)
+	{
+		RegisterConvertibleDelegates(Graph);
+	}
+}
+
 void FKismetCompilerContext::CreateFunctionStubForEvent(UK2Node_Event* SrcEventNode, UObject* OwnerOfTemporaries)
 {
 	FName EventNodeName = GetEventStubFunctionName(SrcEventNode);
@@ -3955,6 +3974,8 @@ void FKismetCompilerContext::CreateAndProcessUbergraph()
 		// Expand out nodes that need it
 		ExpansionStep(ConsolidatedEventGraph, true);
 
+		ReplaceConvertibleDelegates(ConsolidatedEventGraph);
+
 		// If a function in the graph cannot be overridden/placed as event make sure that it is not.
 		VerifyValidOverrideEvent(ConsolidatedEventGraph);
 
@@ -4352,6 +4373,8 @@ void FKismetCompilerContext::ProcessOneFunctionGraph(UEdGraph* SourceGraph, bool
 
 	ExpansionStep(FunctionGraph, false);
 
+	ReplaceConvertibleDelegates(FunctionGraph);
+
 	// Cull the entire construction script graph if after node culling it's trivial, this reduces event spam on object construction:
 	if (SourceGraph->GetFName() == Schema->FN_UserConstructionScript )
 	{
@@ -4688,6 +4711,8 @@ void FKismetCompilerContext::CompileClassLayout(EInternalCompilerFlags InternalF
 		FBlueprintEditorUtils::ConformImplementedInterfaces(Blueprint);
 	}
 
+	// If applicable, register any delegate proxy functions and their captured actor variables
+	RegisterClassDelegateProxiesFromBlueprint();
 
 	// Run thru the class defined variables first, get them registered
 	CreateClassVariablesFromBlueprint();
@@ -4699,6 +4724,9 @@ void FKismetCompilerContext::CompileClassLayout(EInternalCompilerFlags InternalF
 	// Construct a context for each function, doing validation and building the function interface
 	CreateFunctionList();
 
+	// Function list creation should process captured variables. Something went wrong if we missed any.
+	UE_CLOG(!ConvertibleDelegates.IsEmpty(), LogK2Compiler, Warning, TEXT("%d convertible delegates were not processed during class layout compilation."), ConvertibleDelegates.Num());
+
 	// Precompile the functions
 	// Handle delegates signatures first, because they are needed by other functions
 	for (int32 i = 0; i < FunctionList.Num(); ++i)
@@ -5224,15 +5252,15 @@ bool FKismetCompilerContext::ValidateGeneratedClass(UBlueprintGeneratedClass* Cl
 	return UBlueprint::ValidateGeneratedClass(Class);
 }
 
-UEdGraph* FKismetCompilerContext::SpawnIntermediateFunctionGraph(const FString& InDesiredFunctionName)
+UEdGraph* FKismetCompilerContext::SpawnIntermediateFunctionGraph(const FString& InDesiredFunctionName, const UFunction* InSignature, bool bUseUniqueName)
 {
-	FName UniqueGraphName = FBlueprintEditorUtils::FindUniqueKismetName(Blueprint, InDesiredFunctionName);
+	FName UniqueGraphName = bUseUniqueName ? FBlueprintEditorUtils::FindUniqueKismetName(Blueprint, InDesiredFunctionName) : *InDesiredFunctionName;
 
 	UEdGraph* GeneratedFunctionGraph = FBlueprintEditorUtils::CreateNewGraph(Blueprint, UniqueGraphName, UEdGraph::StaticClass(), UEdGraphSchema_K2::StaticClass());
 	GeneratedFunctionGraph->SetFlags(RF_Transient);
 	GeneratedFunctionGraph->bEditable = false;
 
-	FBlueprintEditorUtils::CreateFunctionGraph(Blueprint, GeneratedFunctionGraph, false, (UClass*)nullptr);
+	FBlueprintEditorUtils::CreateFunctionGraph(Blueprint, GeneratedFunctionGraph, false, InSignature);
 
 	// Add the function graph to the list of generated graphs for this compile
 	GeneratedFunctionGraphs.Add(GeneratedFunctionGraph);
@@ -5574,5 +5602,221 @@ void FKismetCompilerContext::ProcessIntermediateTunnelBoundary(UK2Node_Tunnel* T
 	}
 }
 
+void FKismetCompilerContext::RegisterConvertibleDelegates(UEdGraph* Graph)
+{
+	check(Graph);
+
+	TArray<UK2Node_CreateDelegate*> CreateDelegateNodes;
+	Graph->GetNodesOfClass<UK2Node_CreateDelegate>(CreateDelegateNodes);
+
+	for (UK2Node_CreateDelegate* CreateDelegateNode : CreateDelegateNodes)
+	{
+		check(CreateDelegateNode);
+
+		FMemberReference MemberReference;
+		MemberReference.SetDirect(CreateDelegateNode->SelectedFunctionName, CreateDelegateNode->SelectedFunctionGuid, CreateDelegateNode->GetScopeClass(), false);
+
+		const UFunction* BoundFunction = MemberReference.ResolveMember<UFunction>();
+		const UFunction* DelegateSignature = CreateDelegateNode->GetDelegateSignature();
+
+		if (BoundFunction && DelegateSignature)
+		{
+			ConvertibleSignatureMatchResult Result = FKismetCompilerUtilities::DoSignaturesHaveConvertibleFloatTypes(BoundFunction, DelegateSignature);
+			if (Result == ConvertibleSignatureMatchResult::HasConvertibleFloatParams)
+			{
+				int Counter = ConvertibleDelegates.Num();
+
+				FDelegateInfo DelegateInfo;
+				DelegateInfo.ProxyFunctionName = *FString::Printf(TEXT("CREATEDELEGATE_PROXYFUNCTION_%d"), Counter);
+
+				const UEdGraphPin* DelegateNodeSelfPin = CreateDelegateNode->FindPinChecked(UEdGraphSchema_K2::PSC_Self);
+				if (DelegateNodeSelfPin->LinkedTo.Num() > 0)
+				{
+					const UEdGraphPin* TargetActorPin = DelegateNodeSelfPin->LinkedTo[0];
+					check(TargetActorPin);
+
+					DelegateInfo.CapturedVariableName = *FString::Printf(TEXT("CREATEDELEGATE_CAPTUREDVARIABLE_%d"), Counter);
+
+					verify(CreateVariable(DelegateInfo.CapturedVariableName, TargetActorPin->PinType) != nullptr);
+				}
+
+				ConvertibleDelegates.Add(CreateDelegateNode, MoveTemp(DelegateInfo));
+			}
+		}
+	}
+}
+
+void FKismetCompilerContext::ReplaceConvertibleDelegates(UEdGraph* Graph)
+{
+	check(Graph);
+
+	TArray<UK2Node_CreateDelegate*> CreateDelegateNodes;
+	Graph->GetNodesOfClass<UK2Node_CreateDelegate>(CreateDelegateNodes);
+
+	for (UK2Node_CreateDelegate* CreateDelegateNode : CreateDelegateNodes)
+	{
+		check(CreateDelegateNode);
+
+		const UK2Node_CreateDelegate* SourceCreateDelegateNode = Cast<const UK2Node_CreateDelegate>(MessageLog.FindSourceObject(CreateDelegateNode));
+
+		if (const FDelegateInfo* DelegateInfo = ConvertibleDelegates.Find(SourceCreateDelegateNode))
+		{
+			// Create the new function graph
+
+			const UFunction* DelegateSignature = CreateDelegateNode->GetDelegateSignature();
+			UEdGraph* DelegateProxyGraph = SpawnIntermediateFunctionGraph(DelegateInfo->ProxyFunctionName.ToString(), DelegateSignature, false);
+			check(DelegateProxyGraph);
+			const UEdGraphSchema_K2* K2Schema = CastChecked<UEdGraphSchema_K2>(Graph->GetSchema());
+
+			// Create a function call node from the original delegate signature
+
+			FMemberReference MemberReference;
+			MemberReference.SetDirect(CreateDelegateNode->SelectedFunctionName, CreateDelegateNode->SelectedFunctionGuid, CreateDelegateNode->GetScopeClass(), false);
+			const UFunction* BoundFunction = MemberReference.ResolveMember<UFunction>();
+			UK2Node_CallFunction* BoundFunctionNode = SpawnIntermediateNode<UK2Node_CallFunction>(CreateDelegateNode, DelegateProxyGraph);
+			check(BoundFunctionNode);
+			BoundFunctionNode->SetFromFunction(BoundFunction);
+			BoundFunctionNode->AllocateDefaultPins();
+
+			// Link the entry node and the function node pins
+
+			TArray<UK2Node_FunctionEntry*> EntryNodes;
+			DelegateProxyGraph->GetNodesOfClass<UK2Node_FunctionEntry>(EntryNodes);
+			check(EntryNodes.Num() == 1);
+			UK2Node_FunctionEntry* ProxyEntryNode = EntryNodes[0];
+			check(ProxyEntryNode);
+			UEdGraphPin* EntryNodeThenPin = ProxyEntryNode->FindPinChecked(UEdGraphSchema_K2::PN_Then);
+			UEdGraphPin* FunctionNodeExecPin = BoundFunctionNode->GetExecPin();
+			verify(K2Schema->TryCreateConnection(EntryNodeThenPin, FunctionNodeExecPin));
+
+			// Link function entry outputs to the bound function's inputs.
+			// Note that we intentionally skip the "self" pin, since that's a special condition that we handle separately.
+
+			int EntryNodePinCursor = -1;
+			int FunctionNodePinCursor = -1;
+
+			for (int i = 0; i < ProxyEntryNode->Pins.Num(); ++i)
+			{
+				const UEdGraphPin* CurrentPin = ProxyEntryNode->Pins[i];
+				check(CurrentPin);
+				if (!K2Schema->IsMetaPin(*CurrentPin))
+				{
+					EntryNodePinCursor = i;
+					break;
+				}
+			}
+			check(EntryNodePinCursor != -1);
+
+			for (int i = 0; i < BoundFunctionNode->Pins.Num(); ++i)
+			{
+				const UEdGraphPin* CurrentPin = BoundFunctionNode->Pins[i];
+				check(CurrentPin)
+				if (!K2Schema->IsMetaPin(*CurrentPin) && !K2Schema->IsSelfPin(*CurrentPin))
+				{
+					FunctionNodePinCursor = i;
+					break;
+				}
+			}
+			check(FunctionNodePinCursor != -1);
+
+			while (EntryNodePinCursor < ProxyEntryNode->Pins.Num())
+			{
+				UEdGraphPin* OutputPin = ProxyEntryNode->Pins[EntryNodePinCursor];
+				UEdGraphPin* InputPin = BoundFunctionNode->Pins[FunctionNodePinCursor];
+
+				verify(K2Schema->TryCreateConnection(OutputPin, InputPin));
+
+				++EntryNodePinCursor;
+				++FunctionNodePinCursor;
+			}
+
+			// Add and set a capture variable if our self pin is linked to a different actor.
+
+			UEdGraphPin* DelegateNodeSelfPin = CreateDelegateNode->FindPinChecked(UEdGraphSchema_K2::PSC_Self);
+			if (DelegateNodeSelfPin->LinkedTo.Num() > 0)
+			{
+				check(DelegateNodeSelfPin->LinkedTo.Num() == 1);
+
+				// Modify main graph
+				{
+					// Break our delegate node's current self pin
+					// Our replacement proxy func implies the use of 'self'
+
+					UEdGraphPin* TargetActorPin = DelegateNodeSelfPin->LinkedTo[0];
+					check(TargetActorPin);
+					DelegateNodeSelfPin->BreakAllPinLinks();
+
+					// Create our setter node for the captured variable
+
+					UK2Node_VariableSet* SetCapturedVarNode = SpawnIntermediateNode<UK2Node_VariableSet>(CreateDelegateNode, Graph);
+					check(SetCapturedVarNode);
+					SetCapturedVarNode->VariableReference.SetSelfMember(DelegateInfo->CapturedVariableName);
+					SetCapturedVarNode->AllocateDefaultPins();
+
+					// Link the setter node's input to the original actor pin
+
+					UEdGraphPin* SetCapturedVarNodeVariablePin = SetCapturedVarNode->FindPinChecked(DelegateInfo->CapturedVariableName, EGPD_Input);
+					verify(K2Schema->TryCreateConnection(TargetActorPin, SetCapturedVarNodeVariablePin));
+
+					// In order to set the capture variable, we need to find an execution path to insert our setter node.
+					// The best likely candidate is the linked UK2Node_BaseMCDelegate node, which either binds or unbinds a delegate.
+					// We can insert our setter node to its execution pin.
+
+					UEdGraphPin* DelegateOutPin = CreateDelegateNode->GetDelegateOutPin();
+					check(DelegateOutPin);
+					check(DelegateOutPin->LinkedTo.Num() == 1);
+
+					UEdGraphPin* TargetDelegatePin = FBlueprintEditorUtils::FindFirstCompilerRelevantLinkedPin(DelegateOutPin->LinkedTo[0]);
+					check(TargetDelegatePin);
+					UK2Node_BaseMCDelegate* BaseMCDelegateNode = CastChecked<UK2Node_BaseMCDelegate>(TargetDelegatePin->GetOwningNode());
+
+					// Insert the setter node
+
+					UEdGraphPin* SetCapturedVarNodeExecPin = SetCapturedVarNode->GetExecPin();
+					UEdGraphPin* SetCapturedVarNodeThenPin = SetCapturedVarNode->GetThenPin();
+					UEdGraphPin* BaseMCDelegateNodeExecPin = BaseMCDelegateNode->GetExecPin();
+
+					for (UEdGraphPin* LinkedPin : BaseMCDelegateNodeExecPin->LinkedTo)
+					{
+						check(LinkedPin);
+						UEdGraphPin* InputThenPin = FBlueprintEditorUtils::FindFirstCompilerRelevantLinkedPin(LinkedPin);
+						BaseMCDelegateNodeExecPin->BreakLinkTo(InputThenPin);
+
+						verify(K2Schema->TryCreateConnection(InputThenPin, SetCapturedVarNodeExecPin));
+						verify(K2Schema->TryCreateConnection(SetCapturedVarNodeThenPin, BaseMCDelegateNodeExecPin));
+					}
+
+					// Node expansion for the setter node needs to happen, which prunes the default 'get' pin that's normally added during pin allocation
+
+					SetCapturedVarNode->ExpandNode(*this, Graph);
+				}
+
+				// Modify proxy delegate graph
+				{
+					// Create our getter node for the captured variable
+
+					UK2Node_VariableGet* GetCapturedVarNode = SpawnIntermediateNode<UK2Node_VariableGet>(CreateDelegateNode, DelegateProxyGraph);
+					check(GetCapturedVarNode);
+					GetCapturedVarNode->VariableReference.SetSelfMember(DelegateInfo->CapturedVariableName);
+					GetCapturedVarNode->AllocateDefaultPins();
+
+					// Link our getter node to the original function
+
+					UEdGraphPin* BoundFunctionNodeSelfPin = BoundFunctionNode->FindPinChecked(UEdGraphSchema_K2::PN_Self);
+					UEdGraphPin* GetCapturedVarNodeVariablePin = GetCapturedVarNode->FindPinChecked(DelegateInfo->CapturedVariableName, EGPD_Output);
+
+					verify(K2Schema->TryCreateConnection(GetCapturedVarNodeVariablePin, BoundFunctionNodeSelfPin));
+				}
+			}
+
+			// Finally, update the original node to use our new function
+
+			CreateDelegateNode->SetFunction(DelegateProxyGraph->GetFName());
+		}
+
+		ConvertibleDelegates.Remove(SourceCreateDelegateNode);
+	}
+}
+
 #undef LOCTEXT_NAMESPACE
 //////////////////////////////////////////////////////////////////////////
diff --git a/Engine/Source/Editor/KismetCompiler/Private/KismetCompilerMisc.cpp b/Engine/Source/Editor/KismetCompiler/Private/KismetCompilerMisc.cpp
--- a/Engine/Source/Editor/KismetCompiler/Private/KismetCompilerMisc.cpp
+++ b/Engine/Source/Editor/KismetCompiler/Private/KismetCompilerMisc.cpp
@@ -2099,6 +2099,49 @@ bool FKismetCompilerUtilities::CheckFunctionCompiledStatementsThreadSafety(const
 	return bIsThreadSafe;
 }
 
+ConvertibleSignatureMatchResult FKismetCompilerUtilities::DoSignaturesHaveConvertibleFloatTypes(const UFunction* A, const UFunction* B)
+{
+	check(A);
+	check(B);
+
+	if (!A->IsSignatureCompatibleWith(B))
+	{
+		TFieldIterator<FProperty> PropAIt(A);
+		TFieldIterator<FProperty> PropBIt(B);
+
+		while (PropAIt)
+		{
+			if (PropBIt)
+			{
+				if (!FStructUtils::ArePropertiesTheSame(*PropAIt, *PropBIt, false))
+				{
+					bool bHasConvertibleProperties =
+						(PropAIt->IsA<FFloatProperty>() && PropBIt->IsA<FDoubleProperty>()) ||
+						(PropAIt->IsA<FDoubleProperty>() && PropBIt->IsA<FFloatProperty>());
+
+					if (!bHasConvertibleProperties)
+					{
+						return ConvertibleSignatureMatchResult::Different;
+					}
+				}
+			}
+			else
+			{
+				// Mismatched parameter count
+				return ConvertibleSignatureMatchResult::Different;
+			}
+
+			++PropAIt;
+			++PropBIt;
+		}
+
+		// If PropBIt still has parameters, then there was a mismatch
+		return PropBIt ? ConvertibleSignatureMatchResult::Different : ConvertibleSignatureMatchResult::HasConvertibleFloatParams;
+	}
+
+	return ConvertibleSignatureMatchResult::ExactMatch;
+}
+
 //////////////////////////////////////////////////////////////////////////
 // FNodeHandlingFunctor
 
diff --git a/Engine/Source/Editor/KismetCompiler/Public/KismetCompiler.h b/Engine/Source/Editor/KismetCompiler/Public/KismetCompiler.h
--- a/Engine/Source/Editor/KismetCompiler/Public/KismetCompiler.h
+++ b/Engine/Source/Editor/KismetCompiler/Public/KismetCompiler.h
@@ -43,6 +43,7 @@ class UBlueprintGeneratedClass;
 class UClass;
 class UEdGraph;
 class UFunction;
+class UK2Node_CreateDelegate;
 class UK2Node_Event;
 class UK2Node_FunctionEntry;
 class UK2Node_TemporaryVariable;
@@ -104,7 +105,6 @@ protected:
 	/** Event that is broadcast immediately after the function list for this context has been compiled. */
 	FOnFunctionListCompiled FunctionListCompiledEvent;
 
-protected:
 	// This struct holds the various compilation options, such as which passes to perform, whether to save intermediate results, etc
 	FKismetCompilerOptions CompileOptions;
 
@@ -171,13 +171,20 @@ public:
 	// CreateClassVariablesFromBlueprint:
 	bool bAssignDelegateSignatureFunction;
 
+	struct FDelegateInfo
+	{
+		FName ProxyFunctionName;
+		FName CapturedVariableName;
+	};
+
+	TMap<UK2Node_CreateDelegate*, FDelegateInfo> ConvertibleDelegates;
+
 	static FSimpleMulticastDelegate OnPreCompile;
 	static FSimpleMulticastDelegate OnPostCompile;
 
 	/** Broadcasts a notification immediately after the function list for this context has been compiled. */
 	FOnFunctionListCompiled& OnFunctionListCompiled() { return FunctionListCompiledEvent; }
 
-public:
 	FKismetCompilerContext(UBlueprint* SourceSketch, FCompilerResultsLog& InMessageLog, const FKismetCompilerOptions& InCompilerOptions);
 	virtual ~FKismetCompilerContext();
 	
@@ -198,8 +205,9 @@ public:
 
 	const UEdGraphSchema_K2* GetSchema() const { return Schema; }
 
-	/** Spawn an intermediate function graph for this compilation using the specified desired name (which may be modified to make it unique */
-	UEdGraph* SpawnIntermediateFunctionGraph(const FString& InDesiredFunctionName);
+	/** Spawn an intermediate function graph for this compilation using the specified desired name (and optional signature),
+		which may be modified to make it unique. */
+	UEdGraph* SpawnIntermediateFunctionGraph(const FString& InDesiredFunctionName, const UFunction* InSignature = nullptr, bool bUseUniqueName = true);
 
 	// Spawns an intermediate node associated with the source node (for error purposes)
 	template <typename NodeType>
@@ -316,6 +324,15 @@ public:
 	 * Picks the name to use for an autogenerated event stub
 	 */ 
 	FName GetEventStubFunctionName(UK2Node_Event* SrcEventNode);
+
+	/**
+	 * Searches the function graphs and ubergraph pages for any delegate proxies,
+	 * which are then registered with the compiler context.
+	 * If a "captured" variable is needed, then a new property will be added to the current class.
+	 * In this context, a captured variable is any target actor that the delegate will be called on.
+	 */
+	void RegisterClassDelegateProxiesFromBlueprint();
+
 protected:
 	virtual UEdGraphSchema_K2* CreateSchema();
 	virtual void PostCreateSchema();
@@ -371,7 +388,6 @@ protected:
 	/** Prune isolated nodes given the specified graph */
 	void PruneIsolatedNodes(UEdGraph* InGraph, bool bInIncludeNodesThatCouldBeExpandedToRootSet);
 
-protected:
 	// FGraphCompilerContext interface
 	virtual void ValidateLink(const UEdGraphPin* PinA, const UEdGraphPin* PinB) const override;
 	virtual void ValidatePin(const UEdGraphPin* Pin) const override;
@@ -385,7 +401,6 @@ protected:
 		return Pin->PinType.PinCategory != UEdGraphSchema_K2::PC_Exec;
 	}
 
-protected:
 	// Expands out nodes that need it
 	void ExpansionStep(UEdGraph* Graph, bool bAllowUbergraphExpansions);
 
@@ -618,6 +633,25 @@ private:
 	 */
 	void ResetErrorFlags(UEdGraph* Graph) const;
 	
+	/**
+	 * Registers any nodes that are bound to a "convertible" function signature.
+	 * Specifically, these are function signatures that differ only by float/double parameters.
+	 * 
+	 * @param: Reference to graph instance
+	 */
+	void RegisterConvertibleDelegates(UEdGraph* Graph);
+
+	/**
+	 * Modifies the graph to use a proxy delegate function if it uses a convertible delegate signature.
+	 * This involves several steps:
+	 * 1. Creates a new function graph that uses the exact function signature of the delegate.
+	 * 2. Adds and links the original delegate function call, which implicity casts the input parameters.
+	 * 3. If applicable, adds a node to the original graph that sets the variable of the target actor (ie: the captured variable)
+	 * 
+	 * @param: Reference to graph instance
+	 */
+	void ReplaceConvertibleDelegates(UEdGraph* Graph);
+
 	/**
 	 * The compilation manager is a new client of the compiler - it reuses several virtual functions
 	 * that are not usefully public in any other context and so rather than expand the public interface
diff --git a/Engine/Source/Editor/KismetCompiler/Public/KismetCompilerMisc.h b/Engine/Source/Editor/KismetCompiler/Public/KismetCompilerMisc.h
--- a/Engine/Source/Editor/KismetCompiler/Public/KismetCompilerMisc.h
+++ b/Engine/Source/Editor/KismetCompiler/Public/KismetCompilerMisc.h
@@ -37,6 +37,14 @@ struct FKismetFunctionContext;
 //////////////////////////////////////////////////////////////////////////
 // FKismetCompilerUtilities
 
+// Used by DoSignaturesHaveConvertibleFloatTypes
+enum class ConvertibleSignatureMatchResult
+{
+	ExactMatch,					// The function signatures are an exact match
+	HasConvertibleFloatParams,	// The function signatures are identical, except for float/double mismatches, which can be converted
+	Different					// The function signatures are completely different
+};
+
 /** This is a loose collection of utilities used when 'compiling' a new UClass from a K2 graph. */
 class KISMETCOMPILER_API FKismetCompilerUtilities
 {
@@ -136,6 +144,13 @@ public:
 
 	// Helper function used by CheckFunctionThreadSafety. Split out to allow the ability to examine individual compiled statement lists (e.g. for the ubergraph)
 	static bool CheckFunctionCompiledStatementsThreadSafety(const UEdGraphNode* InNode, const UEdGraph* InSourceGraph, const TArray<FBlueprintCompiledStatement*>& InStatements, FCompilerResultsLog& InMessageLog, bool InbEmitErrors = true, TSet<const FBPTerminal*>* InThreadSafeObjectTerms = nullptr);
+
+	/** Similar to UFunction::IsSignatureCompatibleWith, but also checks if the function signatures are convertible.
+	 *
+	 * For example, if a parameter in SourceFunction is of type float, but its corresponding type in OtherFunction is double, then the function is deemed "convertible".
+	 * This is primarily used for binding Blueprint functions with native delegate signatures that use float types.
+	 */
+	static ConvertibleSignatureMatchResult DoSignaturesHaveConvertibleFloatTypes(const UFunction* SourceFunction, const UFunction* OtherFunction);
 };
 
 //////////////////////////////////////////////////////////////////////////
diff --git a/Engine/Source/Runtime/CoreUObject/Public/UObject/Class.h b/Engine/Source/Runtime/CoreUObject/Public/UObject/Class.h
--- a/Engine/Source/Runtime/CoreUObject/Public/UObject/Class.h
+++ b/Engine/Source/Runtime/CoreUObject/Public/UObject/Class.h
@@ -3794,7 +3794,7 @@ inline T* GetMutableDefault(UClass *Class)
 
 struct FStructUtils
 {
-	static bool ArePropertiesTheSame(const FProperty* A, const FProperty* B, bool bCheckPropertiesNames);
+	COREUOBJECT_API static bool ArePropertiesTheSame(const FProperty* A, const FProperty* B, bool bCheckPropertiesNames);
 
 	/** Do structures have exactly the same memory layout */
 	COREUOBJECT_API static bool TheSameLayout(const UStruct* StructA, const UStruct* StructB, bool bCheckPropertiesNames = false);
