Make GIsDuplicatingClassForReinstancing thread-local for the purpose of async loading

diff --git a/Engine/Source/Editor/KismetCompiler/Private/UserDefinedStructureCompilerUtils.cpp b/Engine/Source/Editor/KismetCompiler/Private/UserDefinedStructureCompilerUtils.cpp
--- a/Engine/Source/Editor/KismetCompiler/Private/UserDefinedStructureCompilerUtils.cpp
+++ b/Engine/Source/Editor/KismetCompiler/Private/UserDefinedStructureCompilerUtils.cpp
@@ -91,7 +91,7 @@ struct FUserDefinedStructureCompilerInner
 				const FString ReinstancedName = FString::Printf(TEXT("STRUCT_REINST_%s"), *StructureToReinstance->GetName());
 				const FName UniqueName = MakeUniqueObjectName(GetTransientPackage(), UUserDefinedStruct::StaticClass(), FName(*ReinstancedName));
 
-				TGuardValue<bool> IsDuplicatingClassForReinstancing(GIsDuplicatingClassForReinstancing, true);
+				TGuardValue<FIsDuplicatingClassForReinstancing, bool> IsDuplicatingClassForReinstancing(GIsDuplicatingClassForReinstancing, true);
 				DuplicatedStruct = (UUserDefinedStruct*)StaticDuplicateObject(StructureToReinstance, GetTransientPackage(), UniqueName, ~RF_Transactional); 
 			}
 
@@ -513,7 +513,7 @@ void FUserDefinedStructureCompilerUtils::ReplaceStructWithTempDuplicateByPredica
 			const FString ReinstancedName = FString::Printf(TEXT("STRUCT_REINST_%s"), *StructureToReinstance->GetName());
 			const FName UniqueName = MakeUniqueObjectName(GetTransientPackage(), UUserDefinedStruct::StaticClass(), FName(*ReinstancedName));
 
-			TGuardValue<bool> IsDuplicatingClassForReinstancing(GIsDuplicatingClassForReinstancing, true);
+			TGuardValue<FIsDuplicatingClassForReinstancing, bool> IsDuplicatingClassForReinstancing(GIsDuplicatingClassForReinstancing, true);
 			DuplicatedStruct = (UUserDefinedStruct*)StaticDuplicateObject(StructureToReinstance, GetTransientPackage(), UniqueName, ~RF_Transactional); 
 		}
 
diff --git a/Engine/Source/Runtime/Core/Private/Misc/CoreGlobals.cpp b/Engine/Source/Runtime/Core/Private/Misc/CoreGlobals.cpp
--- a/Engine/Source/Runtime/Core/Private/Misc/CoreGlobals.cpp
+++ b/Engine/Source/Runtime/Core/Private/Misc/CoreGlobals.cpp
@@ -198,7 +198,7 @@ bool					GIsServer						= false;					/* Whether engine was launched as a server,
 bool					GIsCriticalError				= false;					/* An appError() has occured */
 bool					GIsGuarded						= false;					/* Whether execution is happening within main()/WinMain()'s try/catch handler */
 TSAN_ATOMIC(bool)		GIsRunning(false);											/* Whether execution is happening within MainLoop() */
-bool					GIsDuplicatingClassForReinstancing = false;					/* Whether we are currently using SDO on a UClass or CDO for live reinstancing */
+FIsDuplicatingClassForReinstancing	GIsDuplicatingClassForReinstancing;					        /* Whether we are currently using SDO on a UClass or CDO for live reinstancing */
 /** This specifies whether the engine was launched as a build machine process								*/
 bool					GIsBuildMachine					= false;
 /** This determines if we should output any log text.  If Yes then no log text should be emitted.			*/
@@ -725,6 +725,19 @@ DEFINE_LOG_CATEGORY(LogVirtualization);
 	DEFINE_LOG_CATEGORY_HELPER(PLATFORM_GLOBAL_LOG_CATEGORY_ALT);
 #endif
 
+thread_local bool PRIVATE_GIsDuplicatingClassForReinstancing = false;
+
+FIsDuplicatingClassForReinstancing& FIsDuplicatingClassForReinstancing::operator= (bool bOther)
+{
+	PRIVATE_GIsDuplicatingClassForReinstancing = bOther;
+	return *this;
+}
+
+FIsDuplicatingClassForReinstancing::operator bool() const
+{
+	return PRIVATE_GIsDuplicatingClassForReinstancing;
+}
+
 #undef LOCTEXT_NAMESPACE
 
 bool IsRunningCookOnTheFly()
diff --git a/Engine/Source/Runtime/Core/Public/CoreGlobals.h b/Engine/Source/Runtime/Core/Public/CoreGlobals.h
--- a/Engine/Source/Runtime/Core/Public/CoreGlobals.h
+++ b/Engine/Source/Runtime/Core/Public/CoreGlobals.h
@@ -300,12 +300,19 @@ FORCEINLINE bool IsAllowCommandletAudio()
 #endif
 }
 
+class CORE_API FIsDuplicatingClassForReinstancing
+{
+public:
+	FIsDuplicatingClassForReinstancing& operator= (bool bOther);
+	operator bool() const;
+};
+
 extern CORE_API bool GEdSelectionLock;
 extern CORE_API bool GIsClient;
 extern CORE_API bool GIsServer;
 extern CORE_API bool GIsCriticalError;
 extern CORE_API TSAN_ATOMIC(bool) GIsRunning;
-extern CORE_API bool GIsDuplicatingClassForReinstancing;
+extern CORE_API FIsDuplicatingClassForReinstancing GIsDuplicatingClassForReinstancing;
 
 /**
 * These are set when the engine first starts up.
