Fix latent call site targets in a compiled Blueprint ubergraph to avoid re-entry when the latent info term is declared as the first argument.

diff --git a/Engine/Source/Editor/KismetCompiler/Private/KismetCompilerMisc.cpp b/Engine/Source/Editor/KismetCompiler/Private/KismetCompilerMisc.cpp
--- a/Engine/Source/Editor/KismetCompiler/Private/KismetCompilerMisc.cpp
+++ b/Engine/Source/Editor/KismetCompiler/Private/KismetCompilerMisc.cpp
@@ -2462,9 +2462,11 @@ struct FGotoMapUtils
 {
 	static bool IsUberGraphEventStatement(const FBlueprintCompiledStatement* GotoStatement)
 	{
+		// Note: Latent function call sites also utilize the UbergraphCallIndex field, so we need to separate them from ubergraph event targets when the latent term is at index 0.
 		return GotoStatement
 			&& (GotoStatement->Type == KCST_CallFunction) 
-			&& (GotoStatement->UbergraphCallIndex == 0);
+			&& (GotoStatement->UbergraphCallIndex == 0)
+			&& (GotoStatement->FunctionToCall && !GotoStatement->FunctionToCall->HasMetaData(FBlueprintMetadata::MD_Latent));
 	}
 
 	static UEdGraphNode* TargetNodeFromPin(const FBlueprintCompiledStatement* GotoStatement, const UEdGraphPin* ExecNet)
