Purge serializable data before consigning child function objects to the trash container during Blueprint class compilation.

diff --git a/Engine/Source/Editor/KismetCompiler/Private/KismetCompiler.cpp b/Engine/Source/Editor/KismetCompiler/Private/KismetCompiler.cpp
--- a/Engine/Source/Editor/KismetCompiler/Private/KismetCompiler.cpp
+++ b/Engine/Source/Editor/KismetCompiler/Private/KismetCompiler.cpp
@@ -685,24 +685,37 @@ void FKismetCompilerContext::CleanAndSanitizeClass(UBlueprintGeneratedClass* Cla
 			continue;
 		}
 
-		// Compiled function script (bytecode) may contain raw pointers to properties owned by this (or another) BP class. These
-		// fields will be immediately freed after the compilation phase (see UClass::DestroyPropertiesPendingDestruction()), thus
-		// invalidating any references to them in the "old" function object's serialized bytecode. Furthermore, reinstancing won't
-		// update this function's bytecode, as that operation is only applied to a Blueprint class's dependencies, and does not
-		// include "trash" class objects that we're creating here (see FBlueprintCompileReinstancer::UpdateBytecodeReferences()).
-		// As we typically run a GC pass after BP compilation, this normally isn't an issue, because the "trash" class object that
-		// owns this function object will get cleaned up at that point, preventing the "old" function object from being serialized
-		// (e.g. as part of reinstancing an external dependency), and ensuring that we don't encounter one of these "dangling"
-		// FField pointers. However, in certain cases (e.g. batched compiles) we may not run a GC pass in-between each operation,
-		// so to cover that case, we ensure that existing bytecode is fully purged before moving a function to the "trash" class.
+		// Class properties are freed independently of GC, but functions we consign to the trash container object will persist until
+		// the next GC pass, so we must purge serializable data first so we don't leak objects or crash due to invalidated references.
 		if(UFunction* Function = Cast<UFunction>(CurrSubObj))
 		{
+			// Compiled function script (bytecode) may contain raw pointers to properties owned by this (or another) BP class. These
+			// fields will be immediately freed after the compilation phase (see UClass::DestroyPropertiesPendingDestruction()), thus
+			// invalidating any references to them in the "old" function object's serialized bytecode. Furthermore, reinstancing won't
+			// update this function's bytecode, as that operation is only applied to a Blueprint class's dependencies, and does not
+			// include "trash" class objects that we're creating here (see FBlueprintCompileReinstancer::UpdateBytecodeReferences()).
+			// As we typically run a GC pass after BP compilation, this normally isn't an issue, because the "trash" class object that
+			// owns this function object will get cleaned up at that point, preventing the "old" function object from being serialized
+			// (e.g. as part of reinstancing an external dependency), and ensuring that we don't encounter one of these "dangling"
+			// FField pointers. However, in certain cases (e.g. batched compiles) we may not run a GC pass in-between each operation,
+			// so to cover that case, we ensure that existing bytecode is fully purged before moving a function to the "trash" class.
 			Function->Script.Empty();
 
 			// This array will get repopulated as part of constructing the new function object when compiling the class; we don't
 			// want to preserve the old copy, because then the old function object could potentially be identified as a referencer
 			// of a stale struct or a class asset during reference replacement if the previous dependency is subsequently recompiled.
 			Function->ScriptAndPropertyObjectReferences.Empty();
+
+			// We also need to destroy all child properties, as some may contain references to existing objects that can later be
+			// invalidated as a result of compilation or GC, and also because we have cleared the references array above that's used
+			// in ARO for reachability analysis during a GC pass. That means any references to objects owned by this class (e.g.
+			// delegate signatures) are no longer seen as referenced by the function nor the class (due to the PurgeClass() below).
+			// Additionally, references to any recompiled class objects or regenerated properties and functions won't be replaced
+			// during the reinstancing phase, since "trash" class objects are meant for GC and will not be considered as a dependency.
+			Function->DestroyChildPropertiesAndResetPropertyLinks();
+
+			// Re-link to ensure that we also reset any cached data that's based on the (now empty) property list.
+			Function->StaticLink(/*bRelinkExistingProperties =*/ true);
 		}
 
 		FName NewSubobjectName = MakeUniqueObjectName(TransientClass, CurrSubObj->GetClass(), CurrSubObj->GetFName());
