Use fallback struct and permissive Cast to avoid crashing when missing a structure (e.g. because it was deleted)

diff --git a/Engine/Source/Editor/BlueprintGraph/Private/K2Node_GetDataTableRow.cpp b/Engine/Source/Editor/BlueprintGraph/Private/K2Node_GetDataTableRow.cpp
--- a/Engine/Source/Editor/BlueprintGraph/Private/K2Node_GetDataTableRow.cpp
+++ b/Engine/Source/Editor/BlueprintGraph/Private/K2Node_GetDataTableRow.cpp
@@ -155,6 +155,10 @@ UScriptStruct* UK2Node_GetDataTableRow::GetDataTableRowStructType() const
 		if (ResultPin && ResultPin->LinkedTo.Num() > 0)
 		{
 			RowStructType = Cast<UScriptStruct>(ResultPin->LinkedTo[0]->PinType.PinSubCategoryObject.Get());
+			if (RowStructType == nullptr)
+			{
+				RowStructType = GetFallbackStruct();
+			}
 			for (int32 LinkIndex = 1; LinkIndex < ResultPin->LinkedTo.Num(); ++LinkIndex)
 			{
 				UEdGraphPin* Link = ResultPin->LinkedTo[LinkIndex];
diff --git a/Engine/Source/Editor/KismetCompiler/Private/KismetCastingUtils.cpp b/Engine/Source/Editor/KismetCompiler/Private/KismetCastingUtils.cpp
--- a/Engine/Source/Editor/KismetCompiler/Private/KismetCastingUtils.cpp
+++ b/Engine/Source/Editor/KismetCompiler/Private/KismetCastingUtils.cpp
@@ -206,8 +206,8 @@ FConversion GetFloatingPointConversion(const UEdGraphPin& SourcePin, const UEdGr
 		}
 		else if ((SourcePin.PinType.PinValueType.TerminalCategory == UEdGraphSchema_K2::PC_Struct) && (DestinationPin.PinType.PinValueType.TerminalCategory == UEdGraphSchema_K2::PC_Struct))
 		{
-			UScriptStruct* SourceStruct = CastChecked<UScriptStruct>(SourcePin.PinType.PinValueType.TerminalSubCategoryObject.Get());
-			UScriptStruct* DestinationStruct = CastChecked<UScriptStruct>(DestinationPin.PinType.PinValueType.TerminalSubCategoryObject.Get());
+			UScriptStruct* SourceStruct = Cast<UScriptStruct>(SourcePin.PinType.PinValueType.TerminalSubCategoryObject.Get());
+			UScriptStruct* DestinationStruct = Cast<UScriptStruct>(DestinationPin.PinType.PinValueType.TerminalSubCategoryObject.Get());
 
 			if (FStructConversionTable::Get().GetConversionFunction(SourceStruct, DestinationStruct).IsSet())
 			{
