Fix issue where failed blueprint function graph merges would stop it from producing function stubs, related to CL 22205520
This fixes crashes when loading blueprints that inherit from blueprints with compiler errors and calling those functions from a construction script

diff --git a/Engine/Source/Editor/KismetCompiler/Private/KismetCompiler.cpp b/Engine/Source/Editor/KismetCompiler/Private/KismetCompiler.cpp
--- a/Engine/Source/Editor/KismetCompiler/Private/KismetCompiler.cpp
+++ b/Engine/Source/Editor/KismetCompiler/Private/KismetCompiler.cpp
@@ -4389,38 +4389,43 @@ void FKismetCompilerContext::ProcessOneFunctionGraph(UEdGraph* SourceGraph, bool
 		FunctionGraph->Rename(*FunctionGraphName, nullptr, RenameFlags);
 	}
 
+	const int32 SavedErrorCount = MessageLog.NumErrors;
+	bool bIsInvalidFunctionGraph = false;
 	FEdGraphUtilities::MergeChildrenGraphsIn(FunctionGraph, FunctionGraph, /* bRequireSchemaMatch = */ true, /* bInIsCompiling = */ true, &MessageLog);
 
 	// If we failed to merge with any child graphs due to an error, we shouldn't continue processing the intermediate graph.
-	if (MessageLog.NumErrors > 0)
+	if (MessageLog.NumErrors > SavedErrorCount)
 	{
-		return;
+		bIsInvalidFunctionGraph = true;
 	}
+	else
+	{
+		ExpansionStep(FunctionGraph, false);
 
-	ExpansionStep(FunctionGraph, false);
-
-	ReplaceConvertibleDelegates(FunctionGraph);
+		ReplaceConvertibleDelegates(FunctionGraph);
 
-	// Cull the entire construction script graph if after node culling it's trivial, this reduces event spam on object construction:
-	if (SourceGraph->GetFName() == Schema->FN_UserConstructionScript )
-	{
-		if(FKismetCompilerUtilities::IsIntermediateFunctionGraphTrivial(Schema->FN_UserConstructionScript, FunctionGraph))
+		// Cull the entire construction script graph if after node culling it's trivial, this reduces event spam on object construction:
+		if (SourceGraph->GetFName() == Schema->FN_UserConstructionScript )
 		{
-			return;
+			if(FKismetCompilerUtilities::IsIntermediateFunctionGraphTrivial(Schema->FN_UserConstructionScript, FunctionGraph))
+			{
+				return;
+			}
 		}
-	}
 
-	// If a function in the graph cannot be overridden/placed as event make sure that it is not.
-	VerifyValidOverrideFunction(FunctionGraph);
+		// If a function in the graph cannot be overridden/placed as event make sure that it is not.
+		VerifyValidOverrideFunction(FunctionGraph);
+
+		// NOTE: The Blueprint compilation manager generates the skeleton class using a different
+		// code path. We do NOT want ValidateGraphIsWellFormed() ran for skeleton-only compiles here
+		// because it can result in errors (the function hasn't been added to the class yet, etc.)
+		check(CompileOptions.CompileType != EKismetCompileType::SkeletonOnly);
 
-	// NOTE: The Blueprint compilation manager generates the skeleton class using a different
-	// code path. We do NOT want ValidateGraphIsWellFormed() ran for skeleton-only compiles here
-	// because it can result in errors (the function hasn't been added to the class yet, etc.)
-	check(CompileOptions.CompileType != EKismetCompileType::SkeletonOnly);
+		// First do some cursory validation (pin types match, inputs to outputs, pins never point to their parent node, etc...)
+		// If this fails we will "stub" the function graph (if callable) or stop altogether to avoid crashes or infinite loops
+		bIsInvalidFunctionGraph = !ValidateGraphIsWellFormed(FunctionGraph);
+	}
 
-	// First do some cursory validation (pin types match, inputs to outputs, pins never point to their parent node, etc...)
-	// If this fails we will "stub" the function graph (if callable) or stop altogether to avoid crashes or infinite loops
-	const bool bIsInvalidFunctionGraph = !ValidateGraphIsWellFormed(FunctionGraph);
 	if (bIsInvalidFunctionGraph)
 	{
 		if(bInternalFunction)
