Fix BP Compiler Error when compiling ExposeOnSpawn variables with a BlueprintSetter function

diff --git a/Engine/Source/Editor/KismetCompiler/Private/KismetCompilerMisc.cpp b/Engine/Source/Editor/KismetCompiler/Private/KismetCompilerMisc.cpp
--- a/Engine/Source/Editor/KismetCompiler/Private/KismetCompilerMisc.cpp
+++ b/Engine/Source/Editor/KismetCompiler/Private/KismetCompilerMisc.cpp
@@ -27,6 +27,7 @@
 #include "K2Node_CallFunction.h"
 #include "K2Node_CallArrayFunction.h"
 #include "K2Node_CallParentFunction.h"
+#include "K2Node_DynamicCast.h"
 #include "K2Node_ExecutionSequence.h"
 #include "K2Node_FunctionEntry.h"
 #include "K2Node_FunctionResult.h"
@@ -789,10 +790,17 @@ UEdGraphPin* FKismetCompilerUtilities::GenerateAssignmentNodes(class FKismetComp
 			const FString& SetFunctionName = Property->GetMetaData(FBlueprintMetadata::MD_PropertySetFunction);
 			if (!SetFunctionName.IsEmpty())
 			{
-				UClass* OwnerClass = Property->GetOwnerClass();
-				UFunction* SetFunction = OwnerClass->FindFunctionByName(*SetFunctionName);
+				UFunction* SetFunction = ForClass->FindFunctionByName(*SetFunctionName);
 				check(SetFunction);
 
+				// Add a cast node so we can call the Setter function with a pin of the right class
+				UK2Node_DynamicCast* CastNode = CompilerContext.SpawnIntermediateNode<UK2Node_DynamicCast>(SpawnNode, SourceGraph);
+				CastNode->TargetType = const_cast<UClass*>(ForClass);
+				CastNode->SetPurity(true);
+				CastNode->AllocateDefaultPins();
+				CastNode->GetCastSourcePin()->MakeLinkTo(CallBeginResult);
+				CastNode->NotifyPinConnectionListChanged(CastNode->GetCastSourcePin());
+
 				UK2Node_CallFunction* CallFuncNode = CompilerContext.SpawnIntermediateNode<UK2Node_CallFunction>(SpawnNode, SourceGraph);
 				CallFuncNode->SetFromFunction(SetFunction);
 				CallFuncNode->AllocateDefaultPins();
@@ -803,7 +811,7 @@ UEdGraphPin* FKismetCompilerUtilities::GenerateAssignmentNodes(class FKismetComp
 
 				// Connect the new object to the 'object' pin
 				UEdGraphPin* ObjectPin = Schema->FindSelfPin(*CallFuncNode, EGPD_Input);
-				CallBeginResult->MakeLinkTo(ObjectPin);
+				CastNode->GetCastResultPin()->MakeLinkTo(ObjectPin);
 
 				// Move Value pin connections
 				UEdGraphPin* SetFunctionValuePin = nullptr;
