Remove stale DelegateSignatureGraphs before compiling, avoid finding and mutating delegate properties in our super class/inherited classes

diff --git a/Engine/Source/Editor/KismetCompiler/Private/KismetCompiler.cpp b/Engine/Source/Editor/KismetCompiler/Private/KismetCompiler.cpp
--- a/Engine/Source/Editor/KismetCompiler/Private/KismetCompiler.cpp
+++ b/Engine/Source/Editor/KismetCompiler/Private/KismetCompiler.cpp
@@ -2504,7 +2504,10 @@ void FKismetCompilerContext::PrecompileFunction(FKismetFunctionContext& Context,
 		{
 			Context.Function->FunctionFlags |= FUNC_Delegate;
 
-			if (FMulticastDelegateProperty* Property = FindFProperty<FMulticastDelegateProperty>(NewClass, Context.DelegateSignatureName))
+			// We really don't want to find our parent's delegate property and accidentally 
+			// overwrite the signature function, so provide EFieldIterationFlags::None:
+			if (FMulticastDelegateProperty* Property = FindFProperty<FMulticastDelegateProperty>(
+					NewClass, Context.DelegateSignatureName, EFieldIterationFlags::None))
 			{
 				Property->SignatureFunction = Context.Function;
 			}
@@ -4751,6 +4754,10 @@ void FKismetCompilerContext::CompileClassLayout(EInternalCompilerFlags InternalF
 
 		// Conform implemented interfaces here, to ensure we generate all functions required by the interface as stubs
 		FBlueprintEditorUtils::ConformImplementedInterfaces(Blueprint);
+
+		// Make sure we don't have any signature graphs with no corresponding variable - some assets have
+		// managed to get into this state - the UI does not provide a way to fix these objects manually
+		FBlueprintEditorUtils::ConformDelegateSignatureGraphs(Blueprint);
 	}
 
 	// If applicable, register any delegate proxy functions and their captured actor variables
diff --git a/Engine/Source/Editor/UnrealEd/Private/Kismet2/BlueprintEditorUtils.cpp b/Engine/Source/Editor/UnrealEd/Private/Kismet2/BlueprintEditorUtils.cpp
--- a/Engine/Source/Editor/UnrealEd/Private/Kismet2/BlueprintEditorUtils.cpp
+++ b/Engine/Source/Editor/UnrealEd/Private/Kismet2/BlueprintEditorUtils.cpp
@@ -3,6 +3,7 @@
 #include "Kismet2/BlueprintEditorUtils.h"
 #include "Algo/AnyOf.h"
 #include "Algo/Copy.h"
+#include "Algo/RemoveIf.h"
 #include "Algo/Transform.h"
 #include "BlueprintCompilationManager.h"
 #include "UObject/Interface.h"
@@ -7572,6 +7573,17 @@ void FBlueprintEditorUtils::ConformImplementedInterfaces(UBlueprint* Blueprint)
 	}
 }
 
+void FBlueprintEditorUtils::ConformDelegateSignatureGraphs(UBlueprint* Blueprint)
+{
+	Blueprint->DelegateSignatureGraphs.SetNum( Algo::RemoveIf(Blueprint->DelegateSignatureGraphs, 
+		[Blueprint](UEdGraph* Graph)
+		{
+			const int32 VarIndex = FBlueprintEditorUtils::FindNewVariableIndex(Blueprint, Graph->GetFName());
+			return VarIndex == INDEX_NONE;
+		}
+	));
+}
+
 void FBlueprintEditorUtils::ConformAllowDeletionFlag(UBlueprint* Blueprint)
 {
 	for (UEdGraph* Graph : Blueprint->FunctionGraphs)
diff --git a/Engine/Source/Editor/UnrealEd/Public/Kismet2/BlueprintEditorUtils.h b/Engine/Source/Editor/UnrealEd/Public/Kismet2/BlueprintEditorUtils.h
--- a/Engine/Source/Editor/UnrealEd/Public/Kismet2/BlueprintEditorUtils.h
+++ b/Engine/Source/Editor/UnrealEd/Public/Kismet2/BlueprintEditorUtils.h
@@ -1481,6 +1481,9 @@ public:
 	/** Makes sure that all graphs for all interfaces we implement exist, and add if not */
 	static void ConformImplementedInterfaces(UBlueprint* Blueprint);
 
+	/** Makes sure that all delegate graphs have a corresponding variable declaration, removing the graph if not */
+	static void ConformDelegateSignatureGraphs(UBlueprint* Blueprint);
+
 	/** Makes sure that all function graphs are flagged as bAllowDeletion=true, except for construction script and animgraph: */
 	static void ConformAllowDeletionFlag(UBlueprint* Blueprint);
 	
