Removed FKismetFunctionContext::UnsortedSeparateExecutionGroups to avoid deprecation warnings.

diff --git a/Engine/Source/Editor/KismetCompiler/Public/KismetCompiledFunctionContext.h b/Engine/Source/Editor/KismetCompiler/Public/KismetCompiledFunctionContext.h
--- a/Engine/Source/Editor/KismetCompiler/Public/KismetCompiledFunctionContext.h
+++ b/Engine/Source/Editor/KismetCompiler/Public/KismetCompiledFunctionContext.h
@@ -66,11 +66,6 @@ public:
 	// Goto fixup requests (each statement (key) wants to goto the first statement attached to the exec out-pin (value))
 	TMap<FBlueprintCompiledStatement*, UEdGraphPin*> GotoFixupRequestMap;
 
-	// @todo: BP2CPP_remove
-	// Used to split uber graph into subfunctions by C++ backend
-	UE_DEPRECATED(5.0, "This member is no longer in use and will be removed.")
-	TArray<TSet<UEdGraphNode*>> UnsortedSeparateExecutionGroups;
-
 	TIndirectArray<FBPTerminal> Parameters;
 	TIndirectArray<FBPTerminal> Results;
 	TIndirectArray<FBPTerminal> VariableReferences;
