Removed deprecated FKismetFunctionContext::bGeneratingCpp which was generating warnings.

diff --git a/Engine/Source/Editor/KismetCompiler/Private/KismetCompilerMisc.cpp b/Engine/Source/Editor/KismetCompiler/Private/KismetCompilerMisc.cpp
--- a/Engine/Source/Editor/KismetCompiler/Private/KismetCompilerMisc.cpp
+++ b/Engine/Source/Editor/KismetCompiler/Private/KismetCompilerMisc.cpp
@@ -2383,7 +2383,6 @@ FKismetFunctionContext::FKismetFunctionContext(FCompilerResultsLog& InMessageLog
 	, bIsSimpleStubGraphWithNoParams(false)
 	, NetFlags(0)
 	, SourceEventFromStubGraph(nullptr)
-	, bGeneratingCpp(false)		// @todo: BP2CPP_remove
 	, bUseFlowStack(true)
 {
 	NetNameMap = new FNetNameMapping();
diff --git a/Engine/Source/Editor/KismetCompiler/Public/KismetCompiledFunctionContext.h b/Engine/Source/Editor/KismetCompiler/Public/KismetCompiledFunctionContext.h
--- a/Engine/Source/Editor/KismetCompiler/Public/KismetCompiledFunctionContext.h
+++ b/Engine/Source/Editor/KismetCompiler/Public/KismetCompiledFunctionContext.h
@@ -100,11 +100,6 @@ public:
 	struct FNetNameMapping* NetNameMap;
 	bool bAllocatedNetNameMap;
 
-	// @todo: BP2CPP_remove
-	//Skip some optimization. C++ code will be generated in this pass. 
-	UE_DEPRECATED(5.0, "This member is no longer in use and will be removed.")
-	bool bGeneratingCpp;
-
 	//Does this function use requires FlowStack ?
 	bool bUseFlowStack;
 public:
