MVVM: Add the FieldNotify pattern to the base Blueprint. The pattern will be removed from the Widget and from the ViewmodelBase class. A BP variable that has the FieldNotify flag will use a different setter. If the variable is replicated, the OnRep function will broadcast the FieldNotification event. FieldNotify variables cannot be used as a ref argument.

diff --git a/Engine/Source/Editor/BlueprintGraph/Classes/EdGraphSchema_K2.h b/Engine/Source/Editor/BlueprintGraph/Classes/EdGraphSchema_K2.h
--- a/Engine/Source/Editor/BlueprintGraph/Classes/EdGraphSchema_K2.h
+++ b/Engine/Source/Editor/BlueprintGraph/Classes/EdGraphSchema_K2.h
@@ -287,6 +287,9 @@ public:
 
 	/** Function marked as explicitly not thread-safe. Opts function out of class-level thread-safety checks. */
 	static const FName MD_NotThreadSafe;
+
+	/** Metadata to add the property or function to the FieldNotification system. */
+	static const FName MD_FieldNotify;
 	
 private:
 	// This class should never be instantiated
diff --git a/Engine/Source/Editor/BlueprintGraph/Classes/K2Node_VariableSet.h b/Engine/Source/Editor/BlueprintGraph/Classes/K2Node_VariableSet.h
--- a/Engine/Source/Editor/BlueprintGraph/Classes/K2Node_VariableSet.h
+++ b/Engine/Source/Editor/BlueprintGraph/Classes/K2Node_VariableSet.h
@@ -52,6 +52,8 @@ class UK2Node_VariableSet : public UK2Node_Variable
 	BLUEPRINTGRAPH_API FName GetRepNotifyName() const;
 	BLUEPRINTGRAPH_API bool ShouldFlushDormancyOnSet() const;
 	BLUEPRINTGRAPH_API bool IsNetProperty() const;
+	BLUEPRINTGRAPH_API bool IsFieldNotifyProperty() const;
+	BLUEPRINTGRAPH_API bool HasFieldNotificationBroadcast() const;
 
 	static FText GetPropertyTooltip(FProperty const* VariableProperty);
 	static FText GetBlueprintVarTooltip(FBPVariableDescription const& VarDesc);
diff --git a/Engine/Source/Editor/BlueprintGraph/Private/CallFunctionHandler.cpp b/Engine/Source/Editor/BlueprintGraph/Private/CallFunctionHandler.cpp
--- a/Engine/Source/Editor/BlueprintGraph/Private/CallFunctionHandler.cpp
+++ b/Engine/Source/Editor/BlueprintGraph/Private/CallFunctionHandler.cpp
@@ -724,9 +724,12 @@ UFunction* FKCHandler_CallFunction::FindFunction(FKismetFunctionContext& Context
 	return nullptr;
 }
 
-static bool FindMatchingReferencedNetPropertyAndPin(TArray<UEdGraphPin*>& RemainingPins, FProperty* FunctionProperty, FProperty*& NetProperty, UEdGraphPin*& PropertyObjectPin)
+namespace UE::BlueprintGraph::Private
+{
+static bool FindMatchingReferencedNetOrFieldNotifyPropertyAndPin(TArray<UEdGraphPin*>& RemainingPins, FProperty* FunctionProperty, FProperty*& NetProperty, FProperty*& FieldNotifyProperty, UEdGraphPin*& PropertyObjectPin)
 {
 	NetProperty = nullptr;
+	FieldNotifyProperty = nullptr;
 	PropertyObjectPin = nullptr;
 
 	if (UNLIKELY(FunctionProperty->HasAllPropertyFlags(CPF_OutParm | CPF_ReferenceParm) && !FunctionProperty->HasAnyPropertyFlags(CPF_ReturnParm | CPF_ConstParm)))
@@ -735,6 +738,7 @@ static bool FindMatchingReferencedNetPropertyAndPin(TArray<UEdGraphPin*>& Remain
 		{
 			if (FunctionProperty->GetFName() == RemainingPins[i]->PinName)
 			{
+				bool bResult = false;
 				UEdGraphPin* ParamPin = RemainingPins[i];
 				RemainingPins.RemoveAtSwap(i);
 				if (UEdGraphPin* PinToTry = FEdGraphUtilities::GetNetFromPin(ParamPin))
@@ -745,23 +749,32 @@ static bool FindMatchingReferencedNetPropertyAndPin(TArray<UEdGraphPin*>& Remain
 
 					if (UK2Node_VariableGet* GetPropertyNode = Cast<UK2Node_VariableGet>(PinToTry->GetOwningNode()))
 					{
-						FProperty* ToCheck = GetPropertyNode->GetPropertyForVariable();
-						if (UNLIKELY(ToCheck && ToCheck->HasAnyPropertyFlags(CPF_Net)))
+						if (FProperty* ToCheck = GetPropertyNode->GetPropertyForVariable())
 						{
-							NetProperty = ToCheck;
-							PropertyObjectPin = GetPropertyNode->FindPinChecked(UEdGraphSchema_K2::PN_Self);
-							return true;
+							if (UNLIKELY(FKismetCompilerUtilities::IsPropertyUsesFieldNotificationSetValueAndBroadcast(ToCheck)))
+							{
+								FieldNotifyProperty = ToCheck;
+								PropertyObjectPin = GetPropertyNode->FindPinChecked(UEdGraphSchema_K2::PN_Self);
+								bResult = true;
+							}
+							if (UNLIKELY(ToCheck->HasAnyPropertyFlags(CPF_Net)))
+							{
+								NetProperty = ToCheck;
+								PropertyObjectPin = GetPropertyNode->FindPinChecked(UEdGraphSchema_K2::PN_Self);
+								bResult = true;
+							}
 						}
 					}
 				}
 
-				return false;
+				return bResult;
 			}
 		}
 	}
 
 	return false;
 }
+}
 
 void FKCHandler_CallFunction::Transform(FKismetFunctionContext& Context, UEdGraphNode* Node)
 {
@@ -870,6 +883,7 @@ void FKCHandler_CallFunction::Transform(FKismetFunctionContext& Context, UEdGrap
 			if (RemainingPins.Num() > 0)
 			{
 				FProperty* NetProperty = nullptr;
+				FProperty* FieldNotifyProperty = nullptr;
 				UEdGraphPin* PropertyObjectPin = nullptr;
 				UEdGraphPin* OldThenPin = CallFuncNode->GetThenPin();
 
@@ -880,7 +894,7 @@ void FKCHandler_CallFunction::Transform(FKismetFunctionContext& Context, UEdGrap
 				// This is similar to the loop in CreateCallFunction
 				for (TFieldIterator<FProperty> It(Function); It; ++It)
 				{
-					if (FindMatchingReferencedNetPropertyAndPin(RemainingPins, *It, NetProperty, PropertyObjectPin))
+					if (UE::BlueprintGraph::Private::FindMatchingReferencedNetOrFieldNotifyPropertyAndPin(RemainingPins, *It, NetProperty, FieldNotifyProperty, PropertyObjectPin))
 					{
 						if (bIsPure)
 						{
@@ -890,56 +904,64 @@ void FKCHandler_CallFunction::Transform(FKismetFunctionContext& Context, UEdGrap
 							break;
 						}
 
-						if (UEdGraphNode* MarkPropertyDirtyNode = FKCPushModelHelpers::ConstructMarkDirtyNodeForProperty(Context, NetProperty, PropertyObjectPin))
+						if (FieldNotifyProperty)
 						{
-							// bool bWereNodesAdded = false;
-							UEdGraphPin* NewThenPin = MarkPropertyDirtyNode->FindPinChecked(UEdGraphSchema_K2::PN_Then);
-							UEdGraphPin* NewInPin = MarkPropertyDirtyNode->FindPinChecked(UEdGraphSchema_K2::PN_Execute);
+							CompilerContext.MessageLog.Warning(*NSLOCTEXT("KismetCompiler", "RefFieldNotifyProperty_Warning", "@@ references a Field Notify property. The property may not be Broadcast correctly. Consider using a temporary variable.").ToString(), Node);
+						}
 
-							if (ensure(NewThenPin) && ensure(NewInPin))
+						if (NetProperty)
+						{
+							if (UEdGraphNode* MarkPropertyDirtyNode = FKCPushModelHelpers::ConstructMarkDirtyNodeForProperty(Context, NetProperty, PropertyObjectPin))
 							{
-								if (OldThenPin)
-								{
-									NewThenPin->CopyPersistentDataFromOldPin(*OldThenPin);
-									OldThenPin->BreakAllPinLinks();
-									OldThenPin->MakeLinkTo(NewInPin);
+								// bool bWereNodesAdded = false;
+								UEdGraphPin* NewThenPin = MarkPropertyDirtyNode->FindPinChecked(UEdGraphSchema_K2::PN_Then);
+								UEdGraphPin* NewInPin = MarkPropertyDirtyNode->FindPinChecked(UEdGraphSchema_K2::PN_Execute);
 
-									OldThenPin = NewThenPin;
-									// bWereNodesAdded = true;
-								}
-								else
+								if (ensure(NewThenPin) && ensure(NewInPin))
 								{
-									// If there's no then pin, we'll instead insert the dirty nodes before the execution
-									// of function with the reference.
-									// This may do weird things with Latent Nodes, so warn about that.
-
-									if (bIsLatent)
-									{
-										CompilerContext.MessageLog.Warning(*NSLOCTEXT("KismetCompiler", "LatentPushModel_Warning", "@@ is a latent node with references to a net property. The property may not be marked dirty in the correct frame.").ToString(), Node);
-									}
-
-									UEdGraphPin* OldInPin = CallFuncNode->FindPin(UEdGraphSchema_K2::PN_Execute);
-									if (OldInPin)
+									if (OldThenPin)
 									{
-										NewInPin->CopyPersistentDataFromOldPin(*OldInPin);
-										OldInPin->BreakAllPinLinks();
+										NewThenPin->CopyPersistentDataFromOldPin(*OldThenPin);
+										OldThenPin->BreakAllPinLinks();
+										OldThenPin->MakeLinkTo(NewInPin);
 
-										NewThenPin->MakeLinkTo(OldInPin);
 										OldThenPin = NewThenPin;
 										// bWereNodesAdded = true;
 									}
+									else
+									{
+										// If there's no then pin, we'll instead insert the dirty nodes before the execution
+										// of function with the reference.
+										// This may do weird things with Latent Nodes, so warn about that.
+
+										if (bIsLatent)
+										{
+											CompilerContext.MessageLog.Warning(*NSLOCTEXT("KismetCompiler", "LatentPushModel_Warning", "@@ is a latent node with references to a net property. The property may not be marked dirty in the correct frame.").ToString(), Node);
+										}
+
+										UEdGraphPin* OldInPin = CallFuncNode->FindPin(UEdGraphSchema_K2::PN_Execute);
+										if (OldInPin)
+										{
+											NewInPin->CopyPersistentDataFromOldPin(*OldInPin);
+											OldInPin->BreakAllPinLinks();
+
+											NewThenPin->MakeLinkTo(OldInPin);
+											OldThenPin = NewThenPin;
+											// bWereNodesAdded = true;
+										}
+									}
 								}
-							}
 
-							/*
-							if (!bWereNodesAdded)
-							{
-								// TODO: JDN - Reenable this once we have other edge cases worked out.
-								// This warning is confusing / not necessarily actionable, and could contribute to spam,
-								// but no one currently relies on these features.
-								// CompilerContext.MessageLog.Warning(*NSLOCTEXT("KismetCompiler", "PushModelNoDirty_Warning", "@@ has reference to net properties, but we were unable to generate dirty nodes.").ToString(), Node);
+								/*
+								if (!bWereNodesAdded)
+								{
+									// TODO: JDN - Reenable this once we have other edge cases worked out.
+									// This warning is confusing / not necessarily actionable, and could contribute to spam,
+									// but no one currently relies on these features.
+									// CompilerContext.MessageLog.Warning(*NSLOCTEXT("KismetCompiler", "PushModelNoDirty_Warning", "@@ has reference to net properties, but we were unable to generate dirty nodes.").ToString(), Node);
+								}
+								*/
 							}
-							*/
 						}
 					}
 
diff --git a/Engine/Source/Editor/BlueprintGraph/Private/EdGraphSchema_K2.cpp b/Engine/Source/Editor/BlueprintGraph/Private/EdGraphSchema_K2.cpp
--- a/Engine/Source/Editor/BlueprintGraph/Private/EdGraphSchema_K2.cpp
+++ b/Engine/Source/Editor/BlueprintGraph/Private/EdGraphSchema_K2.cpp
@@ -280,6 +280,7 @@ const FName FBlueprintMetadata::MD_Namespace(TEXT("Namespace"));
 
 const FName FBlueprintMetadata::MD_ThreadSafe(TEXT("BlueprintThreadSafe"));
 const FName FBlueprintMetadata::MD_NotThreadSafe(TEXT("NotBlueprintThreadSafe"));
+const FName FBlueprintMetadata::MD_FieldNotify(TEXT("FieldNotify"));
 
 //////////////////////////////////////////////////////////////////////////
 
diff --git a/Engine/Source/Editor/BlueprintGraph/Private/K2Node_FunctionEntry.cpp b/Engine/Source/Editor/BlueprintGraph/Private/K2Node_FunctionEntry.cpp
--- a/Engine/Source/Editor/BlueprintGraph/Private/K2Node_FunctionEntry.cpp
+++ b/Engine/Source/Editor/BlueprintGraph/Private/K2Node_FunctionEntry.cpp
@@ -902,6 +902,17 @@ int32 UK2Node_FunctionEntry::GetFunctionFlags() const
 	return ReturnFlags | ExtraFlags;
 }
 
+namespace UE::BlueprintGraph::Private
+{
+	bool GGenerateFieldNotifyBroadcastForOnRepFunction = true;
+	static FAutoConsoleVariableRef CVarGenerateFieldNotifyBroadcastForOnRepFunction(
+		TEXT("bp.GenerateFieldNotifyBroadcastForOnRepFunction"),
+		GGenerateFieldNotifyBroadcastForOnRepFunction,
+		TEXT("When needed, generate a Broadcast FieldNotification node when the OnRep function is called."),
+		ECVF_Default
+	);
+}
+
 void UK2Node_FunctionEntry::ExpandNode(class FKismetCompilerContext& CompilerContext, UEdGraph* SourceGraph)
 {
 	Super::ExpandNode(CompilerContext, SourceGraph);
@@ -916,6 +927,7 @@ void UK2Node_FunctionEntry::ExpandNode(class FKismetCompilerContext& CompilerCon
 	}
 	
 	UEdGraphPin* LastActiveOutputPin = Pins[0];
+	bool bLinkLastPin = false;
 
 	// Only look for FunctionEntry nodes who were duplicated and have a source object
 	if ( UK2Node_FunctionEntry* OriginalNode = Cast<UK2Node_FunctionEntry>(CompilerContext.MessageLog.FindSourceObject(this)) )
@@ -998,12 +1010,32 @@ void UK2Node_FunctionEntry::ExpandNode(class FKismetCompilerContext& CompilerCon
 			}
 		}
 
-		// Finally, hook up the last node to the old node the function entry node was connected to
-		if(OldStartExecPin)
+		bLinkLastPin = true;
+	}
+
+	if (UE::BlueprintGraph::Private::GGenerateFieldNotifyBroadcastForOnRepFunction)
+	{
+		if (FProperty** RepNotifyProperty = CompilerContext.RepNotifyFunctionMap.Find(FunctionReference.GetMemberName()))
 		{
-			LastActiveOutputPin->MakeLinkTo(OldStartExecPin);
+			if ((*RepNotifyProperty)->HasMetaData(FBlueprintMetadata::MD_FieldNotify))
+			{
+				// Generate node to broadcast when the value is changed from the network
+				TTuple<UEdGraphPin*, UEdGraphPin*> ExecThenPins = FKismetCompilerUtilities::GenerateBroadcastFieldNotificationNode(CompilerContext, SourceGraph, this, *RepNotifyProperty);
+
+				LastActiveOutputPin->BreakAllPinLinks();
+				LastActiveOutputPin->MakeLinkTo(ExecThenPins.Get<0>());
+				LastActiveOutputPin = ExecThenPins.Get<1>();
+
+				bLinkLastPin = true;
+			}
 		}
 	}
+
+	// Finally, hook up the last node to the old node the function entry node was connected to
+	if (bLinkLastPin && OldStartExecPin)
+	{
+		LastActiveOutputPin->MakeLinkTo(OldStartExecPin);
+	}
 }
 
 void UK2Node_FunctionEntry::PostReconstructNode()
diff --git a/Engine/Source/Editor/BlueprintGraph/Private/K2Node_VariableSet.cpp b/Engine/Source/Editor/BlueprintGraph/Private/K2Node_VariableSet.cpp
--- a/Engine/Source/Editor/BlueprintGraph/Private/K2Node_VariableSet.cpp
+++ b/Engine/Source/Editor/BlueprintGraph/Private/K2Node_VariableSet.cpp
@@ -151,17 +151,9 @@ FText UK2Node_VariableSet::GetPropertyTooltip(FProperty const* VariableProperty)
 	FText TextFormat;
 	FFormatNamedArguments Args;
 
-	bool const bHasLocalRepNotify = K2Node_VariableSetImpl::PropertyHasLocalRepNotify(VariableProperty);
-
 	FName VarName = NAME_None;
 	if (VariableProperty != nullptr)
 	{
-		if (bHasLocalRepNotify)
-		{
-			Args.Add(TEXT("ReplicationNotifyName"), FText::FromName(VariableProperty->RepNotifyFunc));
-			TextFormat = LOCTEXT("SetVariableWithRepNotify_Tooltip", "Set the value of variable {VarName} and call {ReplicationNotifyName}");
-		}
-
 		VarName = VariableProperty->GetFName();
 
 		UClass* SourceClass = VariableProperty->GetOwnerClass();
@@ -192,18 +184,44 @@ FText UK2Node_VariableSet::GetPropertyTooltip(FProperty const* VariableProperty)
 			}
 		}
 
+		bool const bHasLocalRepNotify = K2Node_VariableSetImpl::PropertyHasLocalRepNotify(VariableProperty);
+		bool const bFieldNotify = VariableProperty->HasMetaData(FBlueprintMetadata::MD_FieldNotify);
+		if (bHasLocalRepNotify)
+		{
+			Args.Add(TEXT("ReplicationNotifyName"), FText::FromName(VariableProperty->RepNotifyFunc));
+		}
 		if (!SubTooltip.IsEmpty())
 		{
 			Args.Add(TEXT("PropertyTooltip"), SubTooltip);
-			if (bHasLocalRepNotify)
+			if (bHasLocalRepNotify && bFieldNotify)
+			{
+				TextFormat = LOCTEXT("SetVariablePropertyWithRepNotifyWithBroadcast_Tooltip", "Set the value of variable {VarName} and call {ReplicationNotifyName} with broadcast.\n{PropertyTooltip}");
+			}
+			else if (bHasLocalRepNotify)
 			{
 				TextFormat = LOCTEXT("SetVariablePropertyWithRepNotify_Tooltip", "Set the value of variable {VarName} and call {ReplicationNotifyName}\n{PropertyTooltip}");
 			}
+			else if (bFieldNotify)
+			{
+				TextFormat = LOCTEXT("SetVariablePropertyWithBroadcast_Tooltip", "Set the value of variable {VarName} and broadcast.\n{PropertyTooltip}");
+			}
 			else
 			{
 				TextFormat = LOCTEXT("SetVariableProperty_Tooltip", "Set the value of variable {VarName}\n{PropertyTooltip}");
 			}
 		}
+		else if (bHasLocalRepNotify && bFieldNotify)
+		{
+			TextFormat = LOCTEXT("SetVariableWithRepNotifyWithBroadcast_Tooltip", "Set the value of variable {VarName} and call {ReplicationNotifyName} with broadcast.");
+		}
+		else if (bHasLocalRepNotify)
+		{
+			TextFormat = LOCTEXT("SetVariableWithRepNotify_Tooltip", "Set the value of variable {VarName} and call {ReplicationNotifyName}");
+		}
+		else if (bFieldNotify)
+		{
+			TextFormat = LOCTEXT("SetVariableWithBroadcast_Tooltip", "Set the value of variable {VarName} and broadcast");
+		}
 	}
 
 	if (TextFormat.IsEmpty())
@@ -273,10 +291,25 @@ FText UK2Node_VariableSet::GetNodeTitle(ENodeTitleType::Type TitleType) const
 		}
 	}
 
+	const bool bHasLocalRepNotify = HasLocalRepNotify();
+	const bool bIsFieldNotify = IsFieldNotifyProperty();
 	if (NumInputsFound != 1)
 	{
-		return HasLocalRepNotify() ? NSLOCTEXT("K2Node", "SetWithNotify", "Set with Notify") : NSLOCTEXT("K2Node", "Set", "Set");
+		if (bHasLocalRepNotify && bIsFieldNotify)
+		{
+			return NSLOCTEXT("K2Node", "SetWithNotifyWithBroadcast", "Set with Notify and Broadcast");
+		}
+		else if (bHasLocalRepNotify)
+		{
+			return NSLOCTEXT("K2Node", "SetWithNotify", "Set with Notify");
+		}
+		else if (bIsFieldNotify)
+		{
+			return NSLOCTEXT("K2Node", "SetWithBroadcast", "Set with Broadcast");
+		}
+		return NSLOCTEXT("K2Node", "Set", "Set");
 	}
+
 	// @TODO: The variable name mutates as the user makes changes to the 
 	//        underlying property, so until we can catch all those cases, we're
 	//        going to leave this optimization off
@@ -286,10 +319,18 @@ FText UK2Node_VariableSet::GetNodeTitle(ENodeTitleType::Type TitleType) const
 		Args.Add(TEXT("PinName"), FText::FromName(InputPinName));
 
 		// FText::Format() is slow, so we cache this to save on performance
-		if (HasLocalRepNotify())
+		if (bHasLocalRepNotify && bIsFieldNotify)
+		{
+			CachedNodeTitle.SetCachedText(FText::Format(NSLOCTEXT("K2Node", "SetWithNotifyWithBroadcastPinName", "Set with Notify and Broadcast {PinName}"), Args), this);
+		}
+		else if (bHasLocalRepNotify)
 		{
 			CachedNodeTitle.SetCachedText(FText::Format(NSLOCTEXT("K2Node", "SetWithNotifyPinName", "Set with Notify {PinName}"), Args), this);
 		}
+		else if (bIsFieldNotify)
+		{
+			CachedNodeTitle.SetCachedText(FText::Format(NSLOCTEXT("K2Node", "SetWithBroadcastPinName", "Set with Broadcast {PinName}"), Args), this);
+		}
 		else
 		{
 			CachedNodeTitle.SetCachedText(FText::Format(NSLOCTEXT("K2Node", "SetPinName", "Set {PinName}"), Args), this);
@@ -337,6 +378,17 @@ FName UK2Node_VariableSet::GetRepNotifyName() const
 	return NAME_None;
 }
 
+bool UK2Node_VariableSet::IsFieldNotifyProperty() const
+{
+	FProperty* Property = GetPropertyForVariable();
+	return Property ? Property->HasMetaData(FBlueprintMetadata::MD_FieldNotify) : false;
+}
+
+bool UK2Node_VariableSet::HasFieldNotificationBroadcast() const
+{
+	FProperty* Property = GetPropertyForVariable();
+	return Property ? FKismetCompilerUtilities::IsPropertyUsesFieldNotificationSetValueAndBroadcast(Property) : false;
+}
 
 FNodeHandlingFunctor* UK2Node_VariableSet::CreateNodeHandler(FKismetCompilerContext& CompilerContext) const
 {
@@ -445,6 +497,7 @@ void UK2Node_VariableSet::ExpandNode(class FKismetCompilerContext& CompilerConte
 		// If property has a BlueprintSetter accessor, then replace the variable get node with a call function
 		if (VariableProperty)
 		{
+			// todo check with BP team if we need to test if the variable has native Setter
 			const FString& SetFunctionName = VariableProperty->GetMetaData(FBlueprintMetadata::MD_PropertySetFunction);
 			if (!SetFunctionName.IsEmpty())
 			{
@@ -483,6 +536,16 @@ void UK2Node_VariableSet::ExpandNode(class FKismetCompilerContext& CompilerConte
 
 				CompilerContext.MovePinLinksToIntermediate(*FindPin(GetVarName(), EGPD_Input), *SetFunctionValuePin);
 			}
+			else if (HasFieldNotificationBroadcast())
+			{
+				TTuple<UEdGraphPin*, UEdGraphPin*> ExecThenPins = FKismetCompilerUtilities::GenerateFieldNotificationSetNode(CompilerContext, SourceGraph, this, FindPinChecked(UEdGraphSchema_K2::PN_Self), VariableProperty, VariableReference, HasLocalRepNotify(), ShouldFlushDormancyOnSet(), IsNetProperty());
+
+				// Move Exec pin connections
+				CompilerContext.MovePinLinksToIntermediate(*GetExecPin(), *ExecThenPins.Get<0>());
+
+				// Move Then pin connections
+				CompilerContext.MovePinLinksToIntermediate(*FindPinChecked(UEdGraphSchema_K2::PN_Then, EGPD_Output), *ExecThenPins.Get<1>());
+			}
 		}
 	}
 
diff --git a/Engine/Source/Editor/BlueprintGraph/Private/VariableSetHandler.cpp b/Engine/Source/Editor/BlueprintGraph/Private/VariableSetHandler.cpp
--- a/Engine/Source/Editor/BlueprintGraph/Private/VariableSetHandler.cpp
+++ b/Engine/Source/Editor/BlueprintGraph/Private/VariableSetHandler.cpp
@@ -170,7 +170,9 @@ void FKCHandler_VariableSet::Transform(FKismetFunctionContext& Context, UEdGraph
 {
 	// Expands node out to include a (local) call to the RepNotify function if necessary
 	UK2Node_VariableSet* SetNotify = Cast<UK2Node_VariableSet>(Node);
-	if ((SetNotify != NULL))
+
+	// If property is HasFieldNotificationBroadcast, then the net code and broadcast will be executed in native code.
+	if (SetNotify && !SetNotify->HasFieldNotificationBroadcast())
 	{
 		if (SetNotify->ShouldFlushDormancyOnSet())
 		{
diff --git a/Engine/Source/Editor/GraphEditor/Private/KismetNodes/SGraphNodeK2Var.cpp b/Engine/Source/Editor/GraphEditor/Private/KismetNodes/SGraphNodeK2Var.cpp
--- a/Engine/Source/Editor/GraphEditor/Private/KismetNodes/SGraphNodeK2Var.cpp
+++ b/Engine/Source/Editor/GraphEditor/Private/KismetNodes/SGraphNodeK2Var.cpp
@@ -106,10 +106,20 @@ void SGraphNodeK2Var::UpdateGraphNode()
 	if (GraphNode->IsA(UK2Node_VariableSet::StaticClass()))
 	{
 		UK2Node_VariableSet *SetNode = Cast<UK2Node_VariableSet>(GraphNode);
-		if(SetNode->HasLocalRepNotify())
+		const bool bHasLocalRepNotify = SetNode->HasLocalRepNotify();
+		const bool bIsFieldNotify = SetNode->IsFieldNotifyProperty();
+		if (bHasLocalRepNotify && bIsFieldNotify)
+		{
+			TitleText = NSLOCTEXT("GraphEditor", "VariableSetWithNotifyWithBroadcast", "SET w/ Notify&Broadcast");
+		}
+		else if (bHasLocalRepNotify)
 		{
 			TitleText = NSLOCTEXT("GraphEditor", "VariableSetWithNotify", "SET w/ Notify");
 		}
+		else if (bIsFieldNotify)
+		{
+			TitleText = NSLOCTEXT("GraphEditor", "VariableSetWithBroadcast", "SET w/ Broadcast");
+		}
 		else
 		{
 			TitleText = NSLOCTEXT("GraphEditor", "VariableSet", "SET");
diff --git a/Engine/Source/Editor/Kismet/Kismet.Build.cs b/Engine/Source/Editor/Kismet/Kismet.Build.cs
--- a/Engine/Source/Editor/Kismet/Kismet.Build.cs
+++ b/Engine/Source/Editor/Kismet/Kismet.Build.cs
@@ -22,13 +22,14 @@ public class Kismet : ModuleRules
 
 		PrivateDependencyModuleNames.AddRange(
 			new string[] {
-                "AppFramework",
+				"AppFramework",
 				"Core",
 				"CoreUObject",
+				"FieldNotification",
 				"ApplicationCore",
 				"Slate",
 				"SlateCore",
-                "EditorStyle",
+				"EditorStyle",
 				"EditorWidgets",
 				"Engine",
 				"Json",
@@ -44,22 +45,22 @@ public class Kismet : ModuleRules
 				"AnimGraph",
 				"PropertyEditor",
 				"SourceControl",
-                "SharedSettingsWidgets",
-                "InputCore",
+				"SharedSettingsWidgets",
+				"InputCore",
 				"EngineSettings",
-                "Projects",
-                "JsonUtilities",
+				"Projects",
+				"JsonUtilities",
 				"DesktopPlatform",
 				"HotReload",
-                "UMGEditor",
-                "UMG", // for SBlueprintDiff
-                "WorkspaceMenuStructure",
+				"UMGEditor",
+				"UMG", // for SBlueprintDiff
+				"WorkspaceMenuStructure",
 				"DeveloperSettings",
 				"ToolMenus",
 				"SubobjectEditor",
 				"SubobjectDataInterface",
 				"ToolWidgets",
-            }
+			}
 			);
 
         DynamicallyLoadedModuleNames.AddRange(
diff --git a/Engine/Source/Editor/Kismet/Private/BlueprintDetailsCustomization.cpp b/Engine/Source/Editor/Kismet/Private/BlueprintDetailsCustomization.cpp
--- a/Engine/Source/Editor/Kismet/Private/BlueprintDetailsCustomization.cpp
+++ b/Engine/Source/Editor/Kismet/Private/BlueprintDetailsCustomization.cpp
@@ -54,6 +54,8 @@
 #include "IDetailsView.h"
 #include "IDocumentation.h"
 #include "IDocumentationPage.h"
+#include "IFieldNotificationClassDescriptor.h"
+#include "INotifyFieldValueChanged.h"
 #include "ISequencerModule.h"
 #include "Input/DragAndDrop.h"
 #include "Input/Events.h"
@@ -519,6 +521,30 @@ void FBlueprintVarActionDetails::CustomizeDetails( IDetailLayoutBuilder& DetailL
 		.IsEnabled(IsVariableInBlueprint())
 		.ToolTip(ReadOnlyTooltip)
 	];
+	
+	if (BlueprintPtr && FBlueprintEditorUtils::ImplementsInterface(BlueprintPtr, true, UNotifyFieldValueChanged::StaticClass()))
+	{
+		// Show the flag if the class implement the interface but only allow the flag to be changed if the variable is defined in BP
+		const FText ToolTip = LOCTEXT("FieldNotifyToolTip", "Generate a field entry for the Field Notification system.");
+		TSharedPtr<SToolTip> FieldNotificationTooltip = IDocumentation::Get()->CreateToolTip(LOCTEXT("FieldNotifyToolTip", "Generate a field entry for the Field Notification system."), NULL, DocLink, TEXT("FieldNotify"));
+
+		Category.AddCustomRow(LOCTEXT("IsVariableFieldNotifyLabel", "Field Notify"))
+			.NameContent()
+			[
+				SNew(STextBlock)
+				.Text(LOCTEXT("IsVariableFieldNotifyLabel", "Field Notify"))
+				.ToolTip(FieldNotificationTooltip)
+				.Font(IDetailLayoutBuilder::GetDetailFont())
+			]
+			.ValueContent()
+			[
+				SNew(SCheckBox)
+				.IsChecked(this, &FBlueprintVarActionDetails::OnFieldNotifyCheckboxState)
+				.OnCheckStateChanged(this, &FBlueprintVarActionDetails::OnFieldNotifyChanged)
+				.IsEnabled(GetPropertyOwnerBlueprint() && IsABlueprintVariable(VariableProperty) && IsAUserVariable(VariableProperty))
+				.ToolTip(FieldNotificationTooltip)
+			];
+	}
 
 	TSharedPtr<SToolTip> Widget3DTooltip = IDocumentation::Get()->CreateToolTip(LOCTEXT("VariableWidget3D_Tooltip", "When true, allows the user to tweak the vector variable by using a 3D transform widget in the viewport (usable when varible is public/enabled)."), NULL, DocLink, TEXT("Widget3D"));
 
@@ -1879,6 +1905,47 @@ void FBlueprintVarActionDetails::OnReadyOnlyChanged(ECheckBoxState InNewState)
 	FBlueprintEditorUtils::SetBlueprintPropertyReadOnlyFlag(BlueprintObj, VarName, bVariableIsReadOnly);
 }
 
+ECheckBoxState FBlueprintVarActionDetails::OnFieldNotifyCheckboxState() const
+{
+	UBlueprint* const BlueprintObj = GetBlueprintObj();
+	const FName VarName = CachedVariableName;
+
+	if (BlueprintObj && !VarName.IsNone())
+	{
+		const int32 VarIndex = FBlueprintEditorUtils::FindNewVariableIndex(BlueprintObj, VarName);
+		if (VarIndex != INDEX_NONE)
+		{
+			return BlueprintObj->NewVariables[VarIndex].HasMetaData(FBlueprintMetadata::MD_FieldNotify) ? ECheckBoxState::Checked : ECheckBoxState::Unchecked;
+		}
+		else if (BlueprintObj->GeneratedClass && BlueprintObj->GeneratedClass->ImplementsInterface(UNotifyFieldValueChanged::StaticClass()) && BlueprintObj->GeneratedClass->GetDefaultObject())
+		{
+			TScriptInterface<INotifyFieldValueChanged> DefaultObject = BlueprintObj->GeneratedClass->GetDefaultObject();
+			return DefaultObject->GetFieldNotificationDescriptor().GetField(BlueprintObj->GeneratedClass, VarName).IsValid() ? ECheckBoxState::Checked : ECheckBoxState::Unchecked;
+		}
+	}
+	return ECheckBoxState::Unchecked;
+}
+
+void FBlueprintVarActionDetails::OnFieldNotifyChanged(ECheckBoxState InNewState)
+{
+	FName VarName = CachedVariableName;
+
+	// Toggle the flag on the blueprint's version of the variable description, based on state
+	const bool bVariableIsFieldNotify = InNewState == ECheckBoxState::Checked;
+
+	UBlueprint* BlueprintObj = MyBlueprint.Pin()->GetBlueprintObj();
+	if (bVariableIsFieldNotify)
+	{
+		// todo look through graph and check if the variable is used in a FieldNotify function, add that info to the metadata
+		//maybe do that in PreCompileFunction
+		FBlueprintEditorUtils::SetBlueprintVariableMetaData(BlueprintObj, VarName, GetLocalVariableScope(CachedVariableProperty.Get()), FBlueprintMetadata::MD_FieldNotify, TEXT(""));
+	}
+	else
+	{
+		FBlueprintEditorUtils::RemoveBlueprintVariableMetaData(GetBlueprintObj(), VarName, GetLocalVariableScope(CachedVariableProperty.Get()), FBlueprintMetadata::MD_FieldNotify);
+	}
+}
+
 ECheckBoxState FBlueprintVarActionDetails::OnCreateWidgetCheckboxState() const
 {
 	FProperty* Property = CachedVariableProperty.Get();
diff --git a/Engine/Source/Editor/Kismet/Private/BlueprintDetailsCustomization.h b/Engine/Source/Editor/Kismet/Private/BlueprintDetailsCustomization.h
--- a/Engine/Source/Editor/Kismet/Private/BlueprintDetailsCustomization.h
+++ b/Engine/Source/Editor/Kismet/Private/BlueprintDetailsCustomization.h
@@ -223,6 +223,9 @@ private:
 	EVisibility ShowReadOnlyCheckboxVisibilty() const;
 	ECheckBoxState OnReadyOnlyCheckboxState() const;
 	void OnReadyOnlyChanged(ECheckBoxState InNewState);
+	
+	ECheckBoxState OnFieldNotifyCheckboxState() const;
+	void OnFieldNotifyChanged(ECheckBoxState InNewState);
 
 	ECheckBoxState OnCreateWidgetCheckboxState() const;
 	void OnCreateWidgetChanged(ECheckBoxState InNewState);
diff --git a/Engine/Source/Editor/KismetCompiler/KismetCompiler.Build.cs b/Engine/Source/Editor/KismetCompiler/KismetCompiler.Build.cs
--- a/Engine/Source/Editor/KismetCompiler/KismetCompiler.Build.cs
+++ b/Engine/Source/Editor/KismetCompiler/KismetCompiler.Build.cs
@@ -11,14 +11,15 @@ public class KismetCompiler : ModuleRules
 				"Core",
 				"CoreUObject",
 				"Engine",
-                "InputCore",
+				"FieldNotification",
+				"InputCore",
 				"EditorFramework",
 				"UnrealEd",
 				"MovieScene",
 				//"MovieSceneTools",
 				"BlueprintGraph",
 				"AnimGraph",
-                "MessageLog",
+				"MessageLog",
 				"Kismet",
 				"ScriptDisassembler",
 			}
diff --git a/Engine/Source/Editor/KismetCompiler/Private/KismetCompiler.cpp b/Engine/Source/Editor/KismetCompiler/Private/KismetCompiler.cpp
--- a/Engine/Source/Editor/KismetCompiler/Private/KismetCompiler.cpp
+++ b/Engine/Source/Editor/KismetCompiler/Private/KismetCompiler.cpp
@@ -24,6 +24,7 @@
 #include "Engine/UserDefinedStruct.h"
 #include "Blueprint/BlueprintExtension.h"
 #include "EdGraphUtilities.h"
+#include "INotifyFieldValueChanged.h"
 #include "K2Node_AddDelegate.h"
 #include "K2Node_CallFunction.h"
 #include "K2Node_Composite.h"
@@ -1195,6 +1196,7 @@ void FKismetCompilerContext::CreateClassVariablesFromBlueprint()
 	BP_SCOPED_COMPILER_EVENT_STAT(EKismetCompilerStats_CreateClassVariables);
 
 	check(NewClass);
+	check(RepNotifyFunctionMap.IsEmpty());
 
 	// Grab the blueprint variables
 	NewClass->NumReplicatedProperties = 0;	// Keep track of how many replicated variables this blueprint adds
@@ -1205,6 +1207,7 @@ void FKismetCompilerContext::CreateClassVariablesFromBlueprint()
 		// Clear out any existing property guids if there could be local changes. The find code handles inherited variables
 		NewClass->PropertyGuids.Reset();
 	}
+	NewClass->FieldNotifies.Reset();
 
 	for (int32 i = 0; i < Blueprint->NewVariables.Num(); ++i)
 	{
@@ -1229,6 +1232,14 @@ void FKismetCompilerContext::CreateClassVariablesFromBlueprint()
 			NewProperty->RepNotifyFunc = Variable.RepNotifyFunc;
 			NewProperty->SetBlueprintReplicationCondition(Variable.ReplicationCondition);
 
+			if (!NewProperty->RepNotifyFunc.IsNone())
+			{
+				if (TObjectPtr<UEdGraph>* Graph = Blueprint->FunctionGraphs.FindByPredicate([GraphName = Variable.RepNotifyFunc](UEdGraph* Other) { return Other->GetFName() == GraphName; }))
+				{
+					RepNotifyFunctionMap.Emplace(Variable.RepNotifyFunc, NewProperty);
+				}
+			}
+
 			if(!Variable.DefaultValue.IsEmpty())
 			{
 				SetPropertyDefaultValue(NewProperty, Variable.DefaultValue);
@@ -1263,6 +1274,14 @@ void FKismetCompilerContext::CreateClassVariablesFromBlueprint()
 					}
 				}
 			}
+
+			bool bIsValidFieldNotify = Variable.HasMetaData(FBlueprintMetadata::MD_FieldNotify) && NewClass->ImplementsInterface(UNotifyFieldValueChanged::StaticClass());
+			if (bIsValidFieldNotify)
+			{
+				ensure(!NewClass->FieldNotifies.Contains(FFieldNotificationId(NewProperty->GetFName())));
+				NewClass->FieldNotifies.Add(FFieldNotificationId(NewProperty->GetFName()));
+			}
+
 			if (bRebuildPropertyMap)
 			{
 				// Update new class property guid map
@@ -3161,6 +3180,8 @@ void FKismetCompilerContext::FinishCompilingClass(UClass* Class)
 	// Create the default object for this class
 	FKismetCompilerUtilities::CompileDefaultProperties(Class);
 
+	CastChecked<UBlueprintGeneratedClass>(Class)->InitializeFieldNotifies();
+
 	AActor* ActorCDO = Cast<AActor>(Class->GetDefaultObject());
 	if (ActorCDO)
 	{
@@ -4605,6 +4626,7 @@ void FKismetCompilerContext::CompileClassLayout(EInternalCompilerFlags InternalF
 	bIsFullCompile = CompileOptions.DoesRequireBytecodeGeneration() && (Blueprint->BlueprintType != BPTYPE_Interface);
 
 	CallsIntoUbergraph.Empty();
+	RepNotifyFunctionMap.Empty();
 	if (bIsFullCompile)
 	{
 		Blueprint->IntermediateGeneratedGraphs.Empty();
diff --git a/Engine/Source/Editor/KismetCompiler/Private/KismetCompilerMisc.cpp b/Engine/Source/Editor/KismetCompiler/Private/KismetCompilerMisc.cpp
--- a/Engine/Source/Editor/KismetCompiler/Private/KismetCompilerMisc.cpp
+++ b/Engine/Source/Editor/KismetCompiler/Private/KismetCompilerMisc.cpp
@@ -18,9 +18,12 @@
 #include "Engine/MemberReference.h"
 #include "Engine/BlueprintGeneratedClass.h"
 #include "Engine/UserDefinedStruct.h"
+#include "FieldNotification/FieldNotificationLibrary.h"
+#include "INotifyFieldValueChanged.h"
 #include "Kismet2/CompilerResultsLog.h"
 #include "EdGraphUtilities.h"
 #include "EdGraphSchema_K2.h"
+#include "FieldNotificationId.h"
 #include "K2Node.h"
 #include "K2Node_BaseAsyncTask.h"
 #include "K2Node_Event.h"
@@ -31,13 +34,19 @@
 #include "K2Node_ExecutionSequence.h"
 #include "K2Node_FunctionEntry.h"
 #include "K2Node_FunctionResult.h"
+#include "K2Node_MakeArray.h"
+#include "K2Node_MakeStruct.h"
+#include "K2Node_Self.h"
+#include "K2Node_TemporaryVariable.h"
 #include "K2Node_Timeline.h"
 #include "K2Node_Variable.h"
+#include "K2Node_VariableGet.h"
 #include "KismetCastingUtils.h"
 #include "KismetCompiledFunctionContext.h"
 #include "KismetCompiler.h"
 
 #include "K2Node_EnumLiteral.h"
+#include "Kismet/KismetArrayLibrary.h"
 #include "Kismet2/KismetReinstanceUtilities.h"
 #include "Kismet2/BlueprintEditorUtils.h"
 #include "Kismet2/StructureEditorUtils.h"
@@ -828,6 +837,24 @@ UEdGraphPin* FKismetCompilerUtilities::GenerateAssignmentNodes(class FKismetComp
 
 				CompilerContext.MovePinLinksToIntermediate(*OrgPin, *SetFunctionValuePin);
 			}
+			else if (FKismetCompilerUtilities::IsPropertyUsesFieldNotificationSetValueAndBroadcast(Property))
+			{
+				// Add a cast node so we can call the Setter function with a pin of the right class
+				//UK2Node_DynamicCast* CastNode = CompilerContext.SpawnIntermediateNode<UK2Node_DynamicCast>(SpawnNode, SourceGraph);
+				//CastNode->TargetType = const_cast<UClass*>(ForClass);
+				//CastNode->SetPurity(true);
+				//CastNode->AllocateDefaultPins();
+				//CastNode->GetCastSourcePin()->MakeLinkTo(CallBeginResult);
+				//CastNode->NotifyPinConnectionListChanged(CastNode->GetCastSourcePin());
+
+				FMemberReference MemberReference;
+				MemberReference.SetFromField<FProperty>(Property, false);
+				TTuple<UEdGraphPin*, UEdGraphPin*> ExecThenPins = GenerateFieldNotificationSetNode(CompilerContext, SourceGraph, SpawnNode, CallBeginResult, Property, MemberReference, false, false, Property->HasAllPropertyFlags(CPF_Net));
+
+				// Connect this node into the exec chain
+				Schema->TryCreateConnection(LastThen, ExecThenPins.Get<0>());
+				LastThen = ExecThenPins.Get<1>();
+			}
 			else if (UFunction* SetByNameFunction = Schema->FindSetVariableByNameFunction(OrgPin->PinType))
 			{
 				UK2Node_CallFunction* SetVarNode = nullptr;
@@ -900,6 +927,220 @@ UEdGraphPin* FKismetCompilerUtilities::GenerateAssignmentNodes(class FKismetComp
 	return LastThen;
 }
 
+namespace UE::KismetCompiler::Private
+{
+	UK2Node_MakeArray* MakeArrayNodeForFieldNotificationId(FKismetCompilerContext& CompilerContext, UEdGraph* SourceGraph, UEdGraphNode* SourceNode, UEdGraphPin* LinkTo)
+	{
+		UK2Node_MakeArray* MakeArrayNode = CompilerContext.SpawnIntermediateNode<UK2Node_MakeArray>(SourceNode, SourceGraph);
+		MakeArrayNode->AllocateDefaultPins();
+		CompilerContext.MessageLog.NotifyIntermediateObjectCreation(MakeArrayNode, SourceNode);
+
+		// Link the array to the other input pin
+		{
+			UEdGraphPin* ArrayOut = MakeArrayNode->GetOutputPin();
+			ArrayOut->MakeLinkTo(LinkTo);
+			MakeArrayNode->PinConnectionListChanged(ArrayOut);
+		}
+
+		return MakeArrayNode;
+	}
+
+	UK2Node_MakeStruct* MakeFieldNotificationIdStruct(FKismetCompilerContext& CompilerContext, UEdGraph* SourceGraph, UEdGraphNode* SourceNode, const FString& FieldId)
+	{
+		UK2Node_MakeStruct* MakeStruct = CompilerContext.SpawnIntermediateNode<UK2Node_MakeStruct>(SourceNode, SourceGraph);
+		MakeStruct->StructType = FFieldNotificationId::StaticStruct();;
+		MakeStruct->AllocateDefaultPins();
+		MakeStruct->bMadeAfterOverridePinRemoval = true;
+		CompilerContext.MessageLog.NotifyIntermediateObjectCreation(MakeStruct, SourceNode);
+
+		CompilerContext.GetSchema()->TrySetDefaultValue(*MakeStruct->FindPinChecked(GET_MEMBER_NAME_STRING_CHECKED(FFieldNotificationId, FieldName)), FieldId);
+
+		return MakeStruct;
+	}
+
+	void MakeLinkFromMakeStructTo(UK2Node_MakeStruct* MakeStructNode, UEdGraphPin* InputPin)
+	{
+		UEdGraphPin** MakeStructOutPin = MakeStructNode->Pins.FindByPredicate([](UEdGraphPin* OtherPin) { return OtherPin->Direction == EGPD_Output; });
+		check(MakeStructOutPin);
+		(*MakeStructOutPin)->MakeLinkTo(InputPin);
+	}
+
+	void AddMakeStructToMakeArrayNode(UK2Node_MakeArray* MakeArrayNode, UK2Node_MakeStruct* MakeStructNode, int32 Index)
+	{
+		// Find the input pin on the "Make Array" node by index.
+		if (Index > 0)
+		{
+			MakeArrayNode->AddInputPin();
+		}
+
+		const FString PinName = FString::Printf(TEXT("[%d]"), Index);
+		MakeLinkFromMakeStructTo(MakeStructNode, MakeArrayNode->FindPinChecked(PinName));
+	}
+}
+
+TTuple<UEdGraphPin*, UEdGraphPin*> FKismetCompilerUtilities::GenerateFieldNotificationSetNode(FKismetCompilerContext& CompilerContext, UEdGraph* SourceGraph, UEdGraphNode* SourceNode, UEdGraphPin* SelfPin, FProperty* VariableProperty, const FMemberReference& VariableReference, bool bHasLocalRepNotify, bool bShouldFlushDormancyOnSet, bool bIsNetProperty)
+{
+	UClass* OwnerClass = VariableProperty->GetOwnerClass();
+	const FString& FieldNotifyMetaData = VariableProperty->GetMetaData(FBlueprintMetadata::MD_FieldNotify);
+	TArray<FString> OtherFieldNotifyToTrigger;
+	FieldNotifyMetaData.ParseIntoArray(OtherFieldNotifyToTrigger, TEXT("|"), true);
+
+	// Set With Broadcast K2 function
+	UK2Node_CallFunction* CallFuncNode = CompilerContext.SpawnIntermediateNode<UK2Node_CallFunction>(SourceNode, SourceGraph);
+	if (OtherFieldNotifyToTrigger.Num() == 0)
+	{
+		CallFuncNode->SetFromFunction(UFieldNotificationLibrary::StaticClass()->FindFunctionByName(GET_MEMBER_NAME_CHECKED(UFieldNotificationLibrary, SetPropertyValueAndBroadcast)));
+	}
+	else
+	{
+		CallFuncNode->SetFromFunction(UFieldNotificationLibrary::StaticClass()->FindFunctionByName(GET_MEMBER_NAME_CHECKED(UFieldNotificationLibrary, SetPropertyValueAndBroadcastFields)));
+	}
+	CallFuncNode->AllocateDefaultPins();
+
+	const UEdGraphSchema_K2* K2Schema = CompilerContext.GetSchema();
+
+	// Set Self pin connections
+	{
+		if (ensure(SelfPin) && SelfPin->LinkedTo.Num() > 0 && SelfPin->Direction == EEdGraphPinDirection::EGPD_Input)
+		{
+			CompilerContext.CopyPinLinksToIntermediate(*SelfPin, *CallFuncNode->FindPinChecked(FName("Object")));
+			CompilerContext.CopyPinLinksToIntermediate(*SelfPin, *CallFuncNode->FindPinChecked(FName("NetOwner")));
+		}
+		else if (SelfPin && SelfPin->Direction == EEdGraphPinDirection::EGPD_Output)
+		{
+			K2Schema->TryCreateConnection(SelfPin, CallFuncNode->FindPinChecked(FName("Object"), EGPD_Input));
+			K2Schema->TryCreateConnection(SelfPin, CallFuncNode->FindPinChecked(FName("NetOwner"), EGPD_Input));
+		}
+		else
+		{
+			UK2Node_Self* SelfNode = CompilerContext.SpawnIntermediateNode<UK2Node_Self>(SourceNode, SourceGraph);
+			SelfNode->AllocateDefaultPins();
+
+			SelfPin = SelfNode->FindPinChecked(UEdGraphSchema_K2::PN_Self);
+			K2Schema->TryCreateConnection(SelfPin, CallFuncNode->FindPinChecked(FName("Object"), EGPD_Input));
+			K2Schema->TryCreateConnection(SelfPin, CallFuncNode->FindPinChecked(FName("NetOwner"), EGPD_Input));
+		}
+	}
+
+	// OldValue and NewValue pin connections
+	{
+		UK2Node_VariableGet* VariableGetNode = CompilerContext.SpawnIntermediateNode<UK2Node_VariableGet>(SourceNode, SourceGraph);
+		VariableGetNode->VariableReference = VariableReference;
+		VariableGetNode->AllocateDefaultPins();
+		CompilerContext.MessageLog.NotifyIntermediateObjectCreation(VariableGetNode, SourceNode);
+
+		{
+			UEdGraphPin* GetSelfPin = K2Schema->FindSelfPin(*VariableGetNode, EGPD_Input);
+			check(SelfPin && GetSelfPin);
+			if (SelfPin->Direction == EEdGraphPinDirection::EGPD_Output)
+			{
+				K2Schema->TryCreateConnection(SelfPin, GetSelfPin);
+			}
+			else
+			{
+				CompilerContext.CopyPinLinksToIntermediate(*SelfPin, *GetSelfPin);
+			}
+		}
+
+		UEdGraphPin* VariableGetPin = VariableGetNode->FindPinChecked(VariableGetNode->GetVarName(), EGPD_Output);
+		UEdGraphPin* OldValuePin = CallFuncNode->FindPinChecked(FName("OldValue"), EGPD_Input);
+		K2Schema->TryCreateConnection(VariableGetPin, OldValuePin);
+		CallFuncNode->NotifyPinConnectionListChanged(OldValuePin);
+
+		// Force the pin to be the same type (see CustomStructureParam) 
+		UEdGraphPin* NewValuePin = CallFuncNode->FindPinChecked(FName("NewValue"), EGPD_Input);
+		NewValuePin->PinType = OldValuePin->PinType;
+		CompilerContext.CopyPinLinksToIntermediate(*SourceNode->FindPinChecked(VariableReference.GetMemberName(), EGPD_Input), *NewValuePin);
+
+		bool bUseReferenceByRef = NewValuePin->LinkedTo.Num() != 0;
+		K2Schema->TrySetDefaultValue(*CallFuncNode->FindPinChecked(FName("NewValueByRef")), bHasLocalRepNotify ? TEXT("True") : TEXT("False"));
+	}
+
+	// Set Net args pin
+	{
+		K2Schema->TrySetDefaultValue(*CallFuncNode->FindPinChecked(FName("bHasLocalRepNotify")), bHasLocalRepNotify ? TEXT("True") : TEXT("False"));
+		K2Schema->TrySetDefaultValue(*CallFuncNode->FindPinChecked(FName("bShouldFlushDormancyOnSet")), bHasLocalRepNotify ? TEXT("True") : TEXT("False"));
+		K2Schema->TrySetDefaultValue(*CallFuncNode->FindPinChecked(FName("bIsNetProperty")), bHasLocalRepNotify ? TEXT("True") : TEXT("False"));
+	}
+
+	TTuple<UEdGraphPin*, UEdGraphPin*> Result = {CallFuncNode->GetExecPin(), CallFuncNode->GetThenPin()};
+
+	// Assign the other broadcast to generate
+	if (OtherFieldNotifyToTrigger.Num() > 0)
+	{
+		UK2Node_MakeArray*MakeArrayNode = UE::KismetCompiler::Private::MakeArrayNodeForFieldNotificationId(CompilerContext, SourceGraph, SourceNode, CallFuncNode->FindPinChecked(TEXT("ExtraFieldIds")));
+		for (int32 ArgIndex = 0; ArgIndex < OtherFieldNotifyToTrigger.Num(); ++ArgIndex)
+		{
+			const FString& OtherFieldId = OtherFieldNotifyToTrigger[ArgIndex];
+			if (!OtherFieldId.IsEmpty())
+			{
+				// Spawn a "Make Struct" node to create the struct FFieldNotificationId
+				UK2Node_MakeStruct* MakeStuctNode = UE::KismetCompiler::Private::MakeFieldNotificationIdStruct(CompilerContext, SourceGraph, SourceNode, OtherFieldId);
+				UE::KismetCompiler::Private::AddMakeStructToMakeArrayNode(MakeArrayNode, MakeStuctNode, ArgIndex);
+			}
+		}
+	}
+
+	return Result;
+}
+
+TTuple<UEdGraphPin*, UEdGraphPin*> FKismetCompilerUtilities::GenerateBroadcastFieldNotificationNode(FKismetCompilerContext& CompilerContext, UEdGraph* SourceGraph, UEdGraphNode* SourceNode, FProperty* Property)
+{
+	UClass* OwnerClass = Property->GetOwnerClass();
+	const FString& FieldNotifyMetaData = Property->GetMetaData(FBlueprintMetadata::MD_FieldNotify);
+	TArray<FString> OtherFieldNotifyToTrigger;
+	FieldNotifyMetaData.ParseIntoArray(OtherFieldNotifyToTrigger, TEXT("|"), true);
+
+	// Broadcast function
+	UK2Node_CallFunction* CallFuncNode = CompilerContext.SpawnIntermediateNode<UK2Node_CallFunction>(SourceNode, SourceGraph);
+	if (OtherFieldNotifyToTrigger.Num() == 0)
+	{
+		CallFuncNode->SetFromFunction(UFieldNotificationLibrary::StaticClass()->FindFunctionByName(GET_MEMBER_NAME_CHECKED(UFieldNotificationLibrary, BroadcastFieldValueChanged)));
+	}
+	else
+	{
+		CallFuncNode->SetFromFunction(UFieldNotificationLibrary::StaticClass()->FindFunctionByName(GET_MEMBER_NAME_CHECKED(UFieldNotificationLibrary, BroadcastFieldsValueChanged)));
+	}
+	CallFuncNode->AllocateDefaultPins();
+
+	const UEdGraphSchema_K2* K2Schema = CompilerContext.GetSchema();
+
+	// Set Self pin connections
+	{
+		UK2Node_Self* SelfNode = CompilerContext.SpawnIntermediateNode<UK2Node_Self>(SourceNode, SourceGraph);
+		SelfNode->AllocateDefaultPins();
+
+		UEdGraphPin* SelfNodePin = SelfNode->FindPinChecked(UEdGraphSchema_K2::PN_Self);
+		K2Schema->TryCreateConnection(SelfNodePin, CallFuncNode->FindPinChecked(FName("Object"), EGPD_Input));
+	}
+
+	TTuple<UEdGraphPin*, UEdGraphPin*> Result = { CallFuncNode->GetExecPin(), CallFuncNode->GetThenPin() };
+
+	// Assign the FieldId
+	if (OtherFieldNotifyToTrigger.Num() == 0)
+	{
+		UK2Node_MakeStruct* MakeStruct = UE::KismetCompiler::Private::MakeFieldNotificationIdStruct(CompilerContext, SourceGraph, SourceNode, Property->GetName());
+		UE::KismetCompiler::Private::MakeLinkFromMakeStructTo(MakeStruct, CallFuncNode->FindPinChecked(TEXT("FieldId")));
+	}
+	else
+	{
+		OtherFieldNotifyToTrigger.Add(Property->GetName());
+		UK2Node_MakeArray* MakeArrayNode = UE::KismetCompiler::Private::MakeArrayNodeForFieldNotificationId(CompilerContext, SourceGraph, SourceNode, CallFuncNode->FindPinChecked(TEXT("FieldIds")));
+		for (int32 ArgIndex = 0; ArgIndex < OtherFieldNotifyToTrigger.Num(); ++ArgIndex)
+		{
+			const FString& OtherFieldId = OtherFieldNotifyToTrigger[ArgIndex];
+			if (!OtherFieldId.IsEmpty())
+			{
+				// Spawn a "Make Struct" node to create the struct FFieldNotificationId
+				UK2Node_MakeStruct* MakeStuctNode = UE::KismetCompiler::Private::MakeFieldNotificationIdStruct(CompilerContext, SourceGraph, SourceNode, OtherFieldId);
+				UE::KismetCompiler::Private::AddMakeStructToMakeArrayNode(MakeArrayNode, MakeStuctNode, ArgIndex);
+			}
+		}
+	}
+
+	return Result;
+}
+
 void FKismetCompilerUtilities::CreateObjectAssignmentStatement(FKismetFunctionContext& Context, UEdGraphNode* Node, FBPTerminal* SrcTerm, FBPTerminal* DstTerm, UEdGraphPin* DstPin)
 {
 	UClass* InputObjClass = Cast<UClass>(SrcTerm->Type.PinSubCategoryObject.Get());
@@ -1639,6 +1880,15 @@ void FKismetCompilerUtilities::DetectValuesReturnedByRef(const UFunction* Func,
 	}
 }
 
+bool FKismetCompilerUtilities::IsPropertyUsesFieldNotificationSetValueAndBroadcast(const FProperty* Property)
+{
+	return Property->HasMetaData(FBlueprintMetadata::MD_FieldNotify)
+		&& !Property->HasMetaData(FBlueprintMetadata::MD_PropertySetFunction)
+		&& !Property->HasSetter()
+		&& Cast<UBlueprintGeneratedClass>(Property->GetOwnerClass()) != nullptr
+		&& Property->GetOwnerClass()->ImplementsInterface(UNotifyFieldValueChanged::StaticClass());
+}
+
 bool FKismetCompilerUtilities::IsStatementReducible(EKismetCompiledStatementType StatementType)
 {
 	switch (StatementType)
diff --git a/Engine/Source/Editor/KismetCompiler/Public/KismetCompiler.h b/Engine/Source/Editor/KismetCompiler/Public/KismetCompiler.h
--- a/Engine/Source/Editor/KismetCompiler/Public/KismetCompiler.h
+++ b/Engine/Source/Editor/KismetCompiler/Public/KismetCompiler.h
@@ -158,6 +158,9 @@ public:
 	// properties generated by nodes.
 	TMap<UEdGraphNode*, UEdGraphNode*> MacroGeneratedNodes;
 
+	// Map from UProperties to their RepNotify graph
+	TMap<FName, FProperty*> RepNotifyFunctionMap;
+
 	// Map from a name to the number of times it's been 'created' (identical nodes create the same variable names, so they need something appended)
 	FNetNameMapping ClassScopeNetNameMap;
 
diff --git a/Engine/Source/Editor/KismetCompiler/Public/KismetCompilerMisc.h b/Engine/Source/Editor/KismetCompiler/Public/KismetCompilerMisc.h
--- a/Engine/Source/Editor/KismetCompiler/Public/KismetCompilerMisc.h
+++ b/Engine/Source/Editor/KismetCompiler/Public/KismetCompilerMisc.h
@@ -33,6 +33,7 @@ class UStruct;
 struct FBPTerminal;
 struct FEdGraphPinType;
 struct FKismetFunctionContext;
+struct FMemberReference;
 
 //////////////////////////////////////////////////////////////////////////
 // FKismetCompilerUtilities
@@ -112,6 +113,12 @@ public:
 	/** Create 'set var by name' nodes and hook them up - used to set values when components are added or actor are created at run time. Returns the 'last then' pin of the assignment nodes */
 	static UEdGraphPin* GenerateAssignmentNodes( class FKismetCompilerContext& CompilerContext, UEdGraph* SourceGraph, UK2Node_CallFunction* CallBeginSpawnNode, UEdGraphNode* SpawnNode, UEdGraphPin* CallBeginResult, const UClass* ForClass );
 
+	/** Create node that replace regular setter and use the SetPropertyValueAndBroadcast. */
+	static TTuple<UEdGraphPin*, UEdGraphPin*> GenerateFieldNotificationSetNode(FKismetCompilerContext& CompilerContext, UEdGraph* SourceGraph, UEdGraphNode* SourceNode, UEdGraphPin* SelfPin, FProperty* VariableProperty, const FMemberReference& VariableReference, bool bHasLocalRepNotify, bool bShouldFlushDormancyOnSet, bool bIsNetProperty);
+	
+	/** Create node to broadcast a FieldNotification value changed */
+	static TTuple<UEdGraphPin*, UEdGraphPin*> GenerateBroadcastFieldNotificationNode(FKismetCompilerContext& CompilerContext, UEdGraph* SourceGraph, UEdGraphNode* SourceNode, FProperty* Property);
+
 	/** Create Kismet assignment statement with proper object <-> interface cast */
 	static void CreateObjectAssignmentStatement(FKismetFunctionContext& Context, UEdGraphNode* Node, FBPTerminal* SrcTerm, FBPTerminal* DstTerm, UEdGraphPin* DstPin = nullptr);
 
@@ -121,6 +128,9 @@ public:
 	/** Generate an error for non-const output parameters */
 	static void DetectValuesReturnedByRef(const UFunction* Func, const UK2Node * Node, FCompilerResultsLog& MessageLog);
 
+	/** @return true when the property is a BP user variable that should uses UFieldNotificationLibrary::SetPropertyValueAndBroadcast to set its value. */
+	static bool IsPropertyUsesFieldNotificationSetValueAndBroadcast(const FProperty* Property);
+
 	static bool IsStatementReducible(EKismetCompiledStatementType StatementType);
 
 	/**
diff --git a/Engine/Source/Editor/UnrealEd/Private/Kismet2/BlueprintEditorUtils.cpp b/Engine/Source/Editor/UnrealEd/Private/Kismet2/BlueprintEditorUtils.cpp
--- a/Engine/Source/Editor/UnrealEd/Private/Kismet2/BlueprintEditorUtils.cpp
+++ b/Engine/Source/Editor/UnrealEd/Private/Kismet2/BlueprintEditorUtils.cpp
@@ -2579,6 +2579,24 @@ void FBlueprintEditorUtils::FindImplementedInterfaces(const UBlueprint* Blueprin
 	}
 }
 
+bool FBlueprintEditorUtils::ImplementsInterface(const UBlueprint* Blueprint, bool bIncludeInherited, UClass* InterfaceClass)
+{
+	// First test this blueprint
+	if (Blueprint->ImplementedInterfaces.ContainsByPredicate([InterfaceClass](const FBPInterfaceDescription& ImplementedInterface){ return ImplementedInterface.Interface->IsChildOf(InterfaceClass); }))
+	{
+		return true;
+	}
+
+	if (bIncludeInherited)
+	{
+		if (UClass* BlueprintParent = Blueprint->ParentClass)
+		{
+			return BlueprintParent->ImplementsInterface(InterfaceClass);
+		}
+	}
+	return false;
+}
+
 UClass* const FBlueprintEditorUtils::GetOverrideFunctionClass(UBlueprint* Blueprint, const FName FuncName, UFunction** OutFunction)
 {
 	if (!Blueprint->SkeletonGeneratedClass)
@@ -3604,8 +3622,9 @@ bool FBlueprintEditorUtils::DoesSupportTimelines(const UBlueprint* Blueprint)
 
 bool FBlueprintEditorUtils::DoesSupportEventGraphs(const UBlueprint* Blueprint)
 {
-	return Blueprint->BlueprintType == BPTYPE_Normal 
-		|| Blueprint->BlueprintType == BPTYPE_LevelScript;
+	return Blueprint->SupportsEventGraphs()
+		&& (Blueprint->BlueprintType == BPTYPE_Normal
+			|| Blueprint->BlueprintType == BPTYPE_LevelScript);
 }
 
 /** Returns whether or not the blueprint supports implementing interfaces */
diff --git a/Engine/Source/Editor/UnrealEd/Public/Kismet2/BlueprintEditorUtils.h b/Engine/Source/Editor/UnrealEd/Public/Kismet2/BlueprintEditorUtils.h
--- a/Engine/Source/Editor/UnrealEd/Public/Kismet2/BlueprintEditorUtils.h
+++ b/Engine/Source/Editor/UnrealEd/Public/Kismet2/BlueprintEditorUtils.h
@@ -1525,6 +1525,14 @@ public:
 	 */
 	static void FindImplementedInterfaces(const UBlueprint* Blueprint, bool bGetAllInterfaces, TArray<UClass*>& ImplementedInterfaces);
 
+	/**
+	 * Returns true if the interfaces is implemented. It can be implemented directly or inherited.
+	 * @param		Blueprint				The blueprint to check interfaces for
+	 * @param		bIncludeInherited		If true, check in all the implemented and inherited interfaces. False, check in the interfaces implemented directly.
+	 * @param		SomeInterface			The interface to check.
+	*/
+	static bool ImplementsInterface(const UBlueprint* Blueprint, bool bIncludeInherited, UClass* SomeInterface);
+
 	/**
 	 * Finds a unique name with a base of the passed in string, appending numbers as needed
 	 *
diff --git a/Engine/Source/Runtime/Engine/Classes/Engine/BlueprintGeneratedClass.h b/Engine/Source/Runtime/Engine/Classes/Engine/BlueprintGeneratedClass.h
--- a/Engine/Source/Runtime/Engine/Classes/Engine/BlueprintGeneratedClass.h
+++ b/Engine/Source/Runtime/Engine/Classes/Engine/BlueprintGeneratedClass.h
@@ -5,6 +5,7 @@
 #include "Blueprint/BlueprintPropertyGuidProvider.h"
 #include "CoreMinimal.h"
 #include "Containers/IndirectArray.h"
+#include "FieldNotificationId.h"
 #include "Stats/Stats.h"
 #include "UObject/ObjectMacros.h"
 #include "Misc/Guid.h"
@@ -673,6 +674,12 @@ public:
 	UPROPERTY()
 	TArray<FBPComponentClassOverride> ComponentClassOverrides;
 
+	/** The name of the properties with FieldNotify */
+	UPROPERTY()
+	TArray<FFieldNotificationId> FieldNotifies;
+
+	int32 FieldNotifiesStartBitNumber = 0;
+
 	/** 'Simple' construction script - graph of components to instance */
 	UPROPERTY()
 	TObjectPtr<class USimpleConstructionScript> SimpleConstructionScript;
@@ -739,6 +746,10 @@ public:
 	static void CreateComponentsForActor(const UClass* ThisClass, AActor* Actor);
 	static void CreateTimelineComponent(AActor* Actor, const UTimelineTemplate* TimelineTemplate);
 
+	/** Called by the UE::FieldNotification::IClassDescriptor. */
+	void ForEachFieldNotify(TFunctionRef<bool(::UE::FieldNotification::FFieldId FieldId)> Callback, bool bIncludeSuper) const;
+	void InitializeFieldNotifies();
+
 #if WITH_EDITOR
 	/**
 	 * Called during serialization to allow the class to stash any sparse class data that needs to 
@@ -753,12 +764,12 @@ public:
 	virtual void ConformSparseClassData(UObject* Object);
 #endif
 
-	// IBlueprintPropertyGuidProvider interface
+	//~ IBlueprintPropertyGuidProvider interface
 	virtual FName FindBlueprintPropertyNameFromGuid(const FGuid& PropertyGuid) const override final;
 	virtual FGuid FindBlueprintPropertyGuidFromName(const FName PropertyName) const override final;
 	// End IBlueprintPropertyGuidProvider interface
 
-	// UObject interface
+	//~ UObject interface
 	virtual void Serialize(FArchive& Ar) override;
 	virtual void PostLoad() override;
 	virtual void PostInitProperties() override;
@@ -766,7 +777,7 @@ public:
 	virtual void GetAssetRegistryTags(TArray<FAssetRegistryTag>& OutTags) const override;
 #if WITH_EDITOR
 	virtual void PostLoadAssetRegistryTags(const FAssetData& InAssetData, TArray<FAssetRegistryTag>& OutTagsAndValuesToUpdate) const;
-#endif // WITH_EDITOR
+#endif //~ WITH_EDITOR
 	virtual FPrimaryAssetId GetPrimaryAssetId() const override;
 	virtual bool NeedsLoadForServer() const override;
 	virtual bool NeedsLoadForClient() const override;
@@ -776,17 +787,17 @@ public:
 	virtual UClass* RegenerateClass(UClass* ClassToRegenerate, UObject* PreviousCDO) override;
 	virtual void PreSaveRoot(FObjectPreSaveRootContext ObjectSaveContext) override;
 	virtual void PostSaveRoot(FObjectPostSaveRootContext ObjectSaveContext) override;
-#endif	//WITH_EDITOR
+#endif	//~ WITH_EDITOR
 	virtual bool IsAsset() const override;
-	// End UObject interface
+	//~ End UObject interface
 
-	// UClass interface
+	//~ UClass interface
 #if WITH_EDITOR
 	virtual UClass* GetAuthoritativeClass() override;
 	virtual void ConditionalRecompileClass(FUObjectSerializeContext* InLoadContext) override;
 	virtual void FlushCompilationQueueForLevel() override;
 	virtual UObject* GetArchetypeForCDO() const override;
-#endif //WITH_EDITOR
+#endif //~ WITH_EDITOR
 	virtual void SerializeDefaultObject(UObject* Object, FStructuredArchive::FSlot Slot) override;
 	virtual void PostLoadDefaultObject(UObject* Object) override;
 	virtual bool IsFunctionImplementedInScript(FName InFunctionName) const override;
@@ -807,7 +818,7 @@ protected:
 	virtual FName FindPropertyNameFromGuid(const FGuid& PropertyGuid) const override;
 	virtual FGuid FindPropertyGuidFromName(const FName InName) const override;
 	virtual bool ArePropertyGuidsAvailable() const override;
-	// End UClass interface
+	//~ End UClass interface
 
 	/**
 	* Returns a linked list of properties with default values that differ from the parent default object. If non-NULL, only these properties will
diff --git a/Engine/Source/Runtime/Engine/Engine.Build.cs b/Engine/Source/Runtime/Engine/Engine.Build.cs
--- a/Engine/Source/Runtime/Engine/Engine.Build.cs
+++ b/Engine/Source/Runtime/Engine/Engine.Build.cs
@@ -77,6 +77,7 @@ public class Engine : ModuleRules
 				"Core",
 				"CoreOnline",
 				"CoreUObject",
+				"FieldNotification",
 				"NetCore",
 				"ApplicationCore",
 				"ImageCore",
diff --git a/Engine/Source/Runtime/Engine/Private/BlueprintGeneratedClass.cpp b/Engine/Source/Runtime/Engine/Private/BlueprintGeneratedClass.cpp
--- a/Engine/Source/Runtime/Engine/Private/BlueprintGeneratedClass.cpp
+++ b/Engine/Source/Runtime/Engine/Private/BlueprintGeneratedClass.cpp
@@ -17,6 +17,8 @@
 #include "Engine/LevelScriptActor.h"
 #include "Engine/SCS_Node.h"
 #include "Engine/InheritableComponentHandler.h"
+#include "IFieldNotificationClassDescriptor.h"
+#include "INotifyFieldValueChanged.h"
 #include "Net/Core/PushModel/PushModel.h"
 #include "UObject/FortniteMainBranchObjectVersion.h"
 #include "UObject/Package.h" // IWYU pragma: keep
@@ -205,9 +207,28 @@ void UBlueprintGeneratedClass::PostLoad()
 		}
 	}
 
+	InitializeFieldNotifies();
+
 	AssembleReferenceTokenStream(true);
 }
 
+void UBlueprintGeneratedClass::InitializeFieldNotifies()
+{
+	//Initialize the interface with the computed FieldNotifies
+	FieldNotifiesStartBitNumber = 0;
+	if (FieldNotifies.Num() && ImplementsInterface(UNotifyFieldValueChanged::StaticClass()) && ensure(ClassDefaultObject))
+	{
+		int32 NumberOfField = 0;
+		TScriptInterface<INotifyFieldValueChanged>(ClassDefaultObject)->GetFieldNotificationDescriptor().ForEachField(this, [&NumberOfField](::UE::FieldNotification::FFieldId FielId)
+			{
+				++NumberOfField;
+				return true;
+			});
+		FieldNotifiesStartBitNumber = NumberOfField - FieldNotifies.Num();
+		ensureMsgf(FieldNotifiesStartBitNumber >= 0, TEXT("The FieldNotifyStartIndex is negative. The number of field should be positive."));
+	}
+}
+
 void UBlueprintGeneratedClass::GetAssetRegistryTags(TArray<FAssetRegistryTag>& OutTags) const
 {
 	Super::GetAssetRegistryTags(OutTags);
@@ -2434,6 +2455,25 @@ FGuid UBlueprintGeneratedClass::FindPropertyGuidFromName(const FName InName) con
 	return FGuid();
 }
 
+void UBlueprintGeneratedClass::ForEachFieldNotify(TFunctionRef<bool(::UE::FieldNotification::FFieldId FieldId)> Callback, bool bIncludeSuper) const
+{
+	ensureMsgf(FieldNotifiesStartBitNumber >= 0, TEXT("The FieldNotifyStartIndex is negative. The number of field should be positive."));
+	for (int32 Index = 0; Index < FieldNotifies.Num(); ++Index)
+	{
+		if (!Callback(UE::FieldNotification::FFieldId(FieldNotifies[Index].GetFieldName(), Index + FieldNotifiesStartBitNumber)))
+		{
+			return;
+		}
+	}
+	if (bIncludeSuper)
+	{
+		if (UBlueprintGeneratedClass* ParentClass = Cast<UBlueprintGeneratedClass>(GetSuperClass()))
+		{
+			ParentClass->ForEachFieldNotify(Callback, bIncludeSuper);
+		}
+	}
+}
+
 #if WITH_EDITORONLY_DATA
 void UBlueprintGeneratedClass::GetEditorTags(FEditorTags& Tags) const
 {
diff --git a/Engine/Source/Runtime/Engine/Private/FieldNotification/FieldNotificationLibrary.cpp b/Engine/Source/Runtime/Engine/Private/FieldNotification/FieldNotificationLibrary.cpp
new file mode 100644
--- /dev/null
+++ b/Engine/Source/Runtime/Engine/Private/FieldNotification/FieldNotificationLibrary.cpp
@@ -0,0 +1,231 @@
+// Copyright Epic Games, Inc. All Rights Reserved.
+
+#include "FieldNotification/FieldNotificationLibrary.h"
+
+#include "INotifyFieldValueChanged.h"
+#include "GameFramework/Actor.h"
+#include "Net/NetPushModelHelpers.h"
+#include "UObject/Script.h"
+
+#define LOCTEXT_NAMESPACE "FieldNotificationLibrary"
+
+namespace UE::FieldNotification::Private
+{
+bool SetPropertyValueAndBroadcast(void* SourceValuePtr, void* TargetValuePtr, FProperty* TargetProperty, UObject* Self
+	, UObject* NetOwner, bool bHasLocalRepNotify, bool bShouldFlushDormancyOnSet, bool bIsNetProperty, TArrayView<FFieldNotificationId> FieldIds
+	, const UObject* ActiveObject, FFrame& StackFrame)
+{
+	TScriptInterface<INotifyFieldValueChanged> NotifyFieldSelf = Self;
+	if (NotifyFieldSelf.GetObject() == nullptr || NotifyFieldSelf.GetInterface() == nullptr || TargetProperty == nullptr)
+	{
+		const FBlueprintExceptionInfo ExceptionInfo(
+			EBlueprintExceptionType::AccessViolation,
+			LOCTEXT("MissingInputProperty", "Failed to resolve the input parameter for SetPropertyValueAndBroadcast.")
+		);
+		FBlueprintCoreDelegates::ThrowScriptException(ActiveObject, StackFrame, ExceptionInfo);
+		return false;
+	}
+
+	check(SourceValuePtr);
+	check(TargetValuePtr);
+	check(Self);
+
+	//if (Old!=New) then
+	//	FlushNetDormancy
+	//	SetValue
+	//	MarkPropertyDirty
+	//	OnRep (it include BroadcastFieldValueChanged) or BroadcastFieldValueChanged
+
+	bool bResult = !TargetProperty->Identical(TargetValuePtr, SourceValuePtr);
+	if (bResult)
+	{
+		// FlushNetDormancy
+		if (bShouldFlushDormancyOnSet)
+		{
+			// if it has an actor, Call FlushNetDormancy
+			if (AActor* ReplicationOwner = Cast<AActor>(NetOwner))
+			{
+				ReplicationOwner->FlushNetDormancy();
+			}
+		}
+
+		// SetValue
+		TargetProperty->SetValue_InContainer(Self, SourceValuePtr);
+
+		bool bCallBroadcastFieldValueChanged = true;
+		if (bIsNetProperty)
+		{
+			// MarkPropertyDirty
+			if (TargetProperty->RepIndex != INDEX_NONE)
+			{
+				UNetPushModelHelpers::MarkPropertyDirtyFromRepIndex(Self, TargetProperty->RepIndex, TargetProperty->GetFName());
+			}
+
+			// OnRep It will do the BroadcastFieldValueChanged inside the OnRep function because we added it at compile time.
+			if (bHasLocalRepNotify && !TargetProperty->RepNotifyFunc.IsNone())
+			{
+				if (UFunction* Function = Self->GetClass()->FindFunctionByName(TargetProperty->RepNotifyFunc))
+				{
+					if (Function->NumParms == 0 && Function->GetReturnProperty() == nullptr)
+					{
+						Self->ProcessEvent(Function, nullptr);
+						bCallBroadcastFieldValueChanged = false;
+					}
+				}
+			}
+		}
+
+		// BroadcastFieldValueChanged
+		if (bCallBroadcastFieldValueChanged)
+		{
+			for (FFieldNotificationId FieldNotificationId : FieldIds)
+			{
+				UE::FieldNotification::FFieldId FieldId = NotifyFieldSelf->GetFieldNotificationDescriptor().GetField(Self->GetClass(), FieldNotificationId.GetFieldName());
+				if (ensureMsgf(FieldId.IsValid(), TEXT("The field should was validated at compilation.")))
+				{
+					NotifyFieldSelf->BroadcastFieldValueChanged(FieldId);
+				}
+			}
+		}
+	}
+
+	return bResult;
+}
+}
+
+void UFieldNotificationLibrary::BroadcastFieldValueChanged(UObject* Object, FFieldNotificationId InFieldNotificationId)
+{
+	TScriptInterface<INotifyFieldValueChanged> NotifyFieldSelf = Object;
+	if (NotifyFieldSelf.GetObject() != nullptr && NotifyFieldSelf.GetInterface() != nullptr && InFieldNotificationId.IsValid())
+	{
+		const UE::FieldNotification::FFieldId FieldId = NotifyFieldSelf->GetFieldNotificationDescriptor().GetField(Object->GetClass(), InFieldNotificationId.GetFieldName());
+		if (FieldId.IsValid())
+		{
+			NotifyFieldSelf->BroadcastFieldValueChanged(FieldId);
+		}
+	}
+}
+
+void UFieldNotificationLibrary::BroadcastFieldsValueChanged(UObject* Object, TArray<FFieldNotificationId> InFieldNotificationIds)
+{
+	TScriptInterface<INotifyFieldValueChanged> NotifyFieldSelf = Object;
+	if (NotifyFieldSelf.GetObject() != nullptr && NotifyFieldSelf.GetInterface() != nullptr)
+	{
+		for (FFieldNotificationId FieldNotificationId : InFieldNotificationIds)
+		{
+			if (FieldNotificationId.IsValid())
+			{
+				const UE::FieldNotification::FFieldId FieldId = NotifyFieldSelf->GetFieldNotificationDescriptor().GetField(Object->GetClass(), FieldNotificationId.GetFieldName());
+				if (FieldId.IsValid())
+				{
+					NotifyFieldSelf->BroadcastFieldValueChanged(FieldId);
+				}
+			}
+		}
+	}
+}
+
+DEFINE_FUNCTION(UFieldNotificationLibrary::execSetPropertyValueAndBroadcast)
+{
+	P_GET_UBOOL(NewValueByRef);
+
+	Stack.StepCompiledIn<FProperty>(nullptr);
+	FProperty* TargetProperty = Stack.MostRecentProperty;
+	void* TargetValuePtr = Stack.MostRecentPropertyAddress;
+
+	void* SourceValuePtr = nullptr;
+	if (NewValueByRef)
+	{
+		Stack.StepCompiledIn<FProperty>(nullptr);
+		SourceValuePtr = Stack.MostRecentPropertyAddress;
+	}
+	else if (TargetProperty)
+	{
+		// Allocate temporary memory for the variable
+		SourceValuePtr = FMemory_Alloca_Aligned(TargetProperty->GetSize(), TargetProperty->GetMinAlignment());
+		TargetProperty->InitializeValue(SourceValuePtr); // probably not needed since they are double/float
+		Stack.Step(Stack.Object, SourceValuePtr);
+	}
+	else
+	{
+		const FBlueprintExceptionInfo ExceptionInfo(
+			EBlueprintExceptionType::AccessViolation,
+			LOCTEXT("MissingInputProperty", "Failed to resolve the input parameter for SetPropertyValueAndBroadcast.")
+		);
+		FBlueprintCoreDelegates::ThrowScriptException(P_THIS, Stack, ExceptionInfo);
+		return;
+	}
+
+	P_GET_OBJECT(UObject, Self);
+	P_GET_OBJECT(UObject, NetOwner);
+	P_GET_UBOOL(bHasLocalRepNotify);
+	P_GET_UBOOL(bShouldFlushDormancyOnSet);
+	P_GET_UBOOL(bIsNetProperty);
+
+	P_FINISH;
+
+	P_NATIVE_BEGIN
+	FFieldNotificationId FieldNotify = FFieldNotificationId(TargetProperty->GetFName());
+	bool bResult = UE::FieldNotification::Private::SetPropertyValueAndBroadcast(SourceValuePtr, TargetValuePtr, TargetProperty, Self
+		, NetOwner, bHasLocalRepNotify, bShouldFlushDormancyOnSet, bIsNetProperty, TArrayView<FFieldNotificationId>(&FieldNotify, 1), P_THIS, Stack);
+	*(bool*)RESULT_PARAM = bResult;
+	P_NATIVE_END
+
+	if (!NewValueByRef && SourceValuePtr)
+	{
+		TargetProperty->DestroyValue(SourceValuePtr);
+	}
+}
+
+DEFINE_FUNCTION(UFieldNotificationLibrary::execSetPropertyValueAndBroadcastFields)
+{
+	P_GET_UBOOL(NewValueByRef);
+
+	Stack.StepCompiledIn<FProperty>(nullptr);
+	FProperty* TargetProperty = Stack.MostRecentProperty;
+	void* TargetValuePtr = Stack.MostRecentPropertyAddress;
+
+	void* SourceValuePtr = nullptr;
+	if (NewValueByRef)
+	{
+		Stack.StepCompiledIn<FProperty>(nullptr);
+		SourceValuePtr = Stack.MostRecentPropertyAddress;
+	}
+	else if (TargetProperty)
+	{
+		// Allocate temporary memory for the variable
+		SourceValuePtr = FMemory_Alloca_Aligned(TargetProperty->GetSize(), TargetProperty->GetMinAlignment());
+		TargetProperty->InitializeValue(SourceValuePtr); // probably not needed since they are double/float
+		Stack.Step(Stack.Object, SourceValuePtr);
+	}
+	else
+	{
+		const FBlueprintExceptionInfo ExceptionInfo(
+			EBlueprintExceptionType::AccessViolation,
+			LOCTEXT("MissingInputProperty", "Failed to resolve the input parameter for SetPropertyValueAndBroadcast.")
+		);
+		FBlueprintCoreDelegates::ThrowScriptException(P_THIS, Stack, ExceptionInfo);
+		return;
+	}
+
+	P_GET_OBJECT(UObject, Self);
+	P_GET_OBJECT(UObject, NetOwner);
+	P_GET_UBOOL(bHasLocalRepNotify);
+	P_GET_UBOOL(bShouldFlushDormancyOnSet);
+	P_GET_UBOOL(bIsNetProperty);
+	P_GET_TARRAY(FFieldNotificationId, FieldNotifies);
+
+	P_FINISH;
+
+	P_NATIVE_BEGIN
+	FieldNotifies.AddUnique(FFieldNotificationId(TargetProperty->GetFName()));
+	bool bResult = UE::FieldNotification::Private::SetPropertyValueAndBroadcast(SourceValuePtr, TargetValuePtr, TargetProperty, Self
+		, NetOwner, bHasLocalRepNotify, bShouldFlushDormancyOnSet, bIsNetProperty, FieldNotifies, P_THIS, Stack);
+	*(bool*)RESULT_PARAM = bResult;
+	P_NATIVE_END
+
+	if (!NewValueByRef && SourceValuePtr)
+	{
+		TargetProperty->DestroyValue(SourceValuePtr);
+	}
+}
diff --git a/Engine/Source/Runtime/Engine/Public/FieldNotification/FieldNotificationLibrary.h b/Engine/Source/Runtime/Engine/Public/FieldNotification/FieldNotificationLibrary.h
new file mode 100644
--- /dev/null
+++ b/Engine/Source/Runtime/Engine/Public/FieldNotification/FieldNotificationLibrary.h
@@ -0,0 +1,42 @@
+// Copyright Epic Games, Inc. All Rights Reserved.
+
+#pragma once
+
+#include "Kismet/BlueprintFunctionLibrary.h"
+#include "FieldNotificationId.h"
+
+#include "FieldNotificationLibrary.generated.h"
+
+
+/**
+ * The Field Notification system allows a user to know when a property value is modified at runtime (note that it can be a function return value).
+ * The class needs to implement the INotifyFieldValueChanged interface.
+ * The property' setter  usually follows this pattern "if (new value != old value) assign the new value; broadcast that the value changed;".
+ * The Blueprint implementation handles that setter logic automatically with SetPropertyValueAndBroadcast.
+ * When a property value is modified by replication, the RepNotify will call BroadcastFieldValueChanged.
+ * A function can also be a Field Notify. The function needs to be const and return a single value.
+ */
+UCLASS(MinimalAPI)
+class UFieldNotificationLibrary : public UBlueprintFunctionLibrary
+{
+	GENERATED_BODY()
+
+public:
+	/** Broadcast that the Field value changed. */
+	UFUNCTION(BlueprintCallable, Category = "FieldNotification", meta = (FieldNotifyInterfaceParam="Object", DisplayName = "Broadcast Field Value Changed"))
+	static void BroadcastFieldValueChanged(UObject* Object, FFieldNotificationId FieldId);
+
+	/** Broadcast that a list of Field values changed. */
+	UFUNCTION(BlueprintCallable, Category = "FieldNotification", meta = (FieldNotifyInterfaceParam = "Object", DisplayName = "Broadcast Fields Value Changed"))
+	static void BroadcastFieldsValueChanged(UObject* Object, TArray<FFieldNotificationId> FieldIds);
+
+	UFUNCTION(BlueprintCallable, CustomThunk, DisplayName = "Set w/ Broadcast", Category = "FieldNotification", meta = (CustomStructureParam = "NewValue,OldValue", FieldNotifyInterfaceParam = "Object", BlueprintInternalUseOnly = "true"))
+	static bool SetPropertyValueAndBroadcast(bool NewValueByRef, UPARAM(ref) const int32& OldValue, UPARAM(ref) const int32& NewValue, UObject* Object, UObject* NetOwner, bool bHasLocalRepNotify, bool bShouldFlushDormancyOnSet, bool bIsNetProperty);
+
+	UFUNCTION(BlueprintCallable, CustomThunk, DisplayName = "Set w/ Broadcast", Category = "FieldNotification", meta = (CustomStructureParam = "NewValue,OldValue", FieldNotifyInterfaceParam = "Object", BlueprintInternalUseOnly = "true"))
+	static bool SetPropertyValueAndBroadcastFields(bool NewValueByRef, UPARAM(ref) const int32& OldValue, UPARAM(ref) const int32& NewValue, UObject* Object, UObject* NetOwner, bool bHasLocalRepNotify, bool bShouldFlushDormancyOnSet, bool bIsNetProperty, TArray<FFieldNotificationId> ExtraFieldIds);
+
+private:
+	DECLARE_FUNCTION(execSetPropertyValueAndBroadcast);
+	DECLARE_FUNCTION(execSetPropertyValueAndBroadcastFields);
+};
