Fix memory scribble when constructing struct literals with bitfields

diff --git a/Engine/Source/Developer/ScriptDisassembler/Private/ScriptDisassembler.cpp b/Engine/Source/Developer/ScriptDisassembler/Private/ScriptDisassembler.cpp
--- a/Engine/Source/Developer/ScriptDisassembler/Private/ScriptDisassembler.cpp
+++ b/Engine/Source/Developer/ScriptDisassembler/Private/ScriptDisassembler.cpp
@@ -972,6 +972,13 @@ void FKismetBytecodeDisassembler::ProcessCommon(int32& ScriptIndex, EExprToken O
 			}
 			break;
 		}
+	case EX_BitFieldConst:
+		{
+			FProperty* BitProperty = ReadPointer<FProperty>(ScriptIndex);
+			uint8 ConstValue = ReadBYTE(ScriptIndex);
+			Ar.Logf(TEXT("%s $%X: set bit property %s to value %d"), *Indents, (int32)Opcode, *GetNameSafe(BitProperty), ConstValue);
+			break;
+		}
 	case EX_ByteConst:
 		{
 			uint8 ConstValue = ReadBYTE(ScriptIndex);
diff --git a/Engine/Source/Editor/KismetCompiler/Private/KismetCompilerVMBackend.cpp b/Engine/Source/Editor/KismetCompiler/Private/KismetCompilerVMBackend.cpp
--- a/Engine/Source/Editor/KismetCompiler/Private/KismetCompilerVMBackend.cpp
+++ b/Engine/Source/Editor/KismetCompiler/Private/KismetCompilerVMBackend.cpp
@@ -469,6 +469,15 @@ public:
 			return Type && (Type->PinCategory == UEdGraphSchema_K2::PC_Boolean);
 		}
 
+		static bool IsBit(const FProperty* Property)
+		{
+			if (Property && Property->GetOwnerStruct() && !CastFieldChecked<FBoolProperty>(Property)->IsNativeBool())
+			{
+				return true;
+			}
+			return false;
+		}
+
 		static bool IsString(const FEdGraphPinType* Type, const FProperty* Property)
 		{
 			if (Property)
@@ -615,7 +624,7 @@ public:
 		}
 	};
 
-	virtual void EmitTermExpr(FBPTerminal* Term, const FProperty* CoerceProperty = NULL, bool bAllowStaticArray = false)
+	void EmitTermExpr(FBPTerminal* Term, const FProperty* CoerceProperty = NULL, bool bAllowStaticArray = false, bool bCallerRequiresBit = false)
 	{
 		if (Term->bIsLiteral)
 		{
@@ -820,8 +829,29 @@ public:
 			}
 			else if (FLiteralTypeHelper::IsBoolean(&Term->Type, CoerceProperty))
 			{
+				// Bitfields in struct literals were being treated as full bytes, but instructions like
+				// EX_LetBool provide the destination as the CoercePoperty, even though they are going
+				// to allocate a full byte for us to write to. To disambiguate I have added bCallerRequiresBit
+				// when a calling expressing allocates only a single bit for us to write to:
 				bool bValue = Term->Name.ToBool();
-				Writer << (bValue ? EX_True : EX_False);
+				const bool bIsBit = bCallerRequiresBit && FLiteralTypeHelper::IsBit(CoerceProperty);
+				
+				if(bIsBit)
+				{
+					check(CoerceProperty);
+					// FArchive const correctness workaround:
+					FProperty* BitProperty = const_cast<FProperty*>(CoerceProperty);
+					uint8 ValueAsByte = bValue;
+
+					// Emit the literal, with enough information to safely write to CoerceProperty:
+					Writer << EX_BitFieldConst;
+					Writer << BitProperty;
+					Writer << ValueAsByte;
+				}
+				else
+				{
+					Writer << (bValue ? EX_True : EX_False);
+				}
 			}
 			else if (FLiteralTypeHelper::IsName(&Term->Type, CoerceProperty))
 			{
@@ -946,7 +976,7 @@ public:
 								NewTerm.ObjectLiteral = CastField<FObjectProperty>(Prop)->GetObjectPropertyValue(Prop->ContainerPtrToValuePtr<void>(StructData));
 							}
 
-							EmitTermExpr(&NewTerm, Prop, true);
+							EmitTermExpr(&NewTerm, Prop, true, true);
 						}
 					}
 					Struct->DestroyStruct(StructData, ArrayDim);
diff --git a/Engine/Source/Runtime/CoreUObject/Private/UObject/ScriptCore.cpp b/Engine/Source/Runtime/CoreUObject/Private/UObject/ScriptCore.cpp
--- a/Engine/Source/Runtime/CoreUObject/Private/UObject/ScriptCore.cpp
+++ b/Engine/Source/Runtime/CoreUObject/Private/UObject/ScriptCore.cpp
@@ -3781,6 +3781,17 @@ DEFINE_FUNCTION(UObject::execMapConst)
 }
 IMPLEMENT_VM_FUNCTION(EX_MapConst, execMapConst);
 
+DEFINE_FUNCTION(UObject::execBitFieldConst)
+{
+	FBoolProperty* BitProperty = CastFieldChecked<FBoolProperty>((FField*)Stack.ReadObject());
+	uint8 ByteValue = Stack.Read<uint8>();
+	// we could pack the bit into the lower bits of the FProperty pointer, but this instruction is rarely used
+	// and it's likely that a simple implementation will be appreciated by readers, debuggers, and even optimizers:
+	checkSlow(ByteValue == 0 || ByteValue == 1); 
+	BitProperty->SetPropertyValue(RESULT_PARAM, (bool)ByteValue);
+}
+IMPLEMENT_VM_FUNCTION(EX_BitFieldConst, execBitFieldConst);
+
 DEFINE_FUNCTION(UObject::execIntZero)
 {
 	*(int32*)RESULT_PARAM = 0;
diff --git a/Engine/Source/Runtime/CoreUObject/Public/UObject/Object.h b/Engine/Source/Runtime/CoreUObject/Public/UObject/Object.h
--- a/Engine/Source/Runtime/CoreUObject/Public/UObject/Object.h
+++ b/Engine/Source/Runtime/CoreUObject/Public/UObject/Object.h
@@ -1632,6 +1632,7 @@ public:
 	DECLARE_FUNCTION(execArrayConst);
 	DECLARE_FUNCTION(execSetConst);
 	DECLARE_FUNCTION(execMapConst);
+	DECLARE_FUNCTION(execBitFieldConst);
 
 	// Object construction
 	DECLARE_FUNCTION(execNew);
diff --git a/Engine/Source/Runtime/CoreUObject/Public/UObject/Script.h b/Engine/Source/Runtime/CoreUObject/Public/UObject/Script.h
--- a/Engine/Source/Runtime/CoreUObject/Public/UObject/Script.h
+++ b/Engine/Source/Runtime/CoreUObject/Public/UObject/Script.h
@@ -201,7 +201,7 @@ enum EExprToken : uint8
 	//						= 0x0E,
 	EX_Let					= 0x0F,	// Assign an arbitrary size value to a variable.
 	//						= 0x10,
-	//						= 0x11,
+	EX_BitFieldConst		= 0x11, // assign to a single bit, defined by an FProperty
 	EX_ClassContext			= 0x12,	// Class default object context.
 	EX_MetaCast             = 0x13, // Metaclass cast.
 	EX_LetBool				= 0x14, // Let boolean variable.
diff --git a/Engine/Source/Runtime/CoreUObject/Public/UObject/ScriptSerialization.h b/Engine/Source/Runtime/CoreUObject/Public/UObject/ScriptSerialization.h
--- a/Engine/Source/Runtime/CoreUObject/Public/UObject/ScriptSerialization.h
+++ b/Engine/Source/Runtime/CoreUObject/Public/UObject/ScriptSerialization.h
@@ -523,6 +523,12 @@
 			while (SerializeExpr(iCode, Ar) != EX_EndMapConst);
 			break;
 		}
+		case EX_BitFieldConst:
+		{
+			XFERPTR(FProperty*);	// Bit property
+			XFER(uint8);			// bit value
+			break;
+		}
 		case EX_ByteConst:
 		case EX_IntConstByte:
 		{
