Warn if a BP Field notify variable is replicated but not RepNotify.

diff --git a/Engine/Source/Editor/KismetCompiler/Private/KismetCompiler.cpp b/Engine/Source/Editor/KismetCompiler/Private/KismetCompiler.cpp
--- a/Engine/Source/Editor/KismetCompiler/Private/KismetCompiler.cpp
+++ b/Engine/Source/Editor/KismetCompiler/Private/KismetCompiler.cpp
@@ -1263,12 +1263,12 @@ void FKismetCompilerContext::CreateClassVariablesFromBlueprint()
 				if (Entry.DataKey == FBlueprintMetadata::MD_ExposeOnSpawn)
 				{
 					NewProperty->SetPropertyFlags(CPF_ExposeOnSpawn);
-					if (NewProperty->HasAnyPropertyFlags(CPF_DisableEditOnInstance))
+					if (NewProperty->HasAnyPropertyFlags(CPF_DisableEditOnInstance) && NewClass == Blueprint->SkeletonGeneratedClass)
 					{
 						MessageLog.Warning(
 							*FText::Format(
 								LOCTEXT("ExposeToSpawnButPrivateWarningFmt", "Variable {0} is marked as 'Expose on Spawn' but not marked as 'Instance Editable'; please make it 'Instance Editable'"),
-								FText::FromString(NewProperty->GetName())
+								FText::FromName(NewProperty->GetFName())
 							).ToString()
 						);
 					}
@@ -1282,6 +1282,16 @@ void FKismetCompilerContext::CreateClassVariablesFromBlueprint()
 				NewClass->FieldNotifies.Add(FFieldNotificationId(NewProperty->GetFName()));
 			}
 
+			if (NewProperty->HasAnyPropertyFlags(CPF_Net) && NewProperty->RepNotifyFunc.IsNone() && bIsValidFieldNotify && NewClass == Blueprint->SkeletonGeneratedClass)
+			{
+				MessageLog.Warning(
+					*FText::Format(
+						LOCTEXT("ReplicatedFieldNotifyWarning", "Variable {0} is FieldNotify and replicated but not marked as 'RepNotify'; please make it 'RepNotify'."),
+						FText::FromName(NewProperty->GetFName())
+					).ToString()
+				);
+			}
+
 			if (bRebuildPropertyMap)
 			{
 				// Update new class property guid map
