(Newly created) Blueprint Child Classes were not marked as Transactional and therefore would not undo properly.  The fix is to mark any Blueprint CDO as Transactional.

diff --git a/Engine/Source/Editor/KismetCompiler/Private/KismetCompilerMisc.cpp b/Engine/Source/Editor/KismetCompiler/Private/KismetCompilerMisc.cpp
--- a/Engine/Source/Editor/KismetCompiler/Private/KismetCompilerMisc.cpp
+++ b/Engine/Source/Editor/KismetCompiler/Private/KismetCompilerMisc.cpp
@@ -619,6 +619,18 @@ void FKismetCompilerUtilities::CompileDefaultProperties(UClass* Class)
 {
 	UObject* DefaultObject = Class->GetDefaultObject(); // Force the default object to be constructed if it isn't already
 	check(DefaultObject);
+
+	// This is a Blueprint class which means the CDO is exposed in the Editor and thus should be transactional.
+	TArray<UObject*> AllSubObjects;
+	constexpr bool bIncludeNestedObjects = true;
+	constexpr EObjectFlags ExclusionFlags = RF_Transient;
+	GetObjectsWithOuter(DefaultObject, AllSubObjects, bIncludeNestedObjects, ExclusionFlags);
+
+	DefaultObject->SetFlags(RF_Transactional);
+	for (UObject* SubObject : AllSubObjects)
+	{
+		SubObject->SetFlags(RF_Transactional);
+	}
 }
 
 void FKismetCompilerUtilities::LinkAddedProperty(UStruct* Structure, FProperty* NewProperty)
