Make explicit and validate ublueprint and ubpgc classes, honor the mappings when creating a blueprint from a UClass

diff --git a/Engine/Source/Editor/AnimGraph/Private/AnimGraphModule.cpp b/Engine/Source/Editor/AnimGraph/Private/AnimGraphModule.cpp
--- a/Engine/Source/Editor/AnimGraph/Private/AnimGraphModule.cpp
+++ b/Engine/Source/Editor/AnimGraph/Private/AnimGraphModule.cpp
@@ -56,6 +56,8 @@ void FAnimGraphModule::StartupModule()
 	{
 		return MakeShared<FAnimBlueprintCompilerContext>(CastChecked<UAnimBlueprint>(InBlueprint), InMessageLog, InCompileOptions);
 	});
+	IKismetCompilerInterface& KismetCompilerModule = FModuleManager::LoadModuleChecked<IKismetCompilerInterface>("KismetCompiler");
+	KismetCompilerModule.OverrideBPTypeForClass(UAnimInstance::StaticClass(), UAnimBlueprint::StaticClass());
 
 	IModularFeatures::Get().RegisterModularFeature("PropertyAccessBlueprintBinding", &PropertyAccessAnimBlueprintBinding);
 
diff --git a/Engine/Source/Editor/Blutility/Private/BlutilityModule.cpp b/Engine/Source/Editor/Blutility/Private/BlutilityModule.cpp
--- a/Engine/Source/Editor/Blutility/Private/BlutilityModule.cpp
+++ b/Engine/Source/Editor/Blutility/Private/BlutilityModule.cpp
@@ -14,8 +14,11 @@
 #include "Editor.h"
 #include "Editor/EditorEngine.h"
 #include "EditorSupportDelegates.h"
+#include "EditorUtilityActor.h"
 #include "EditorUtilityBlueprint.h"
+#include "EditorUtilityCamera.h"
 #include "EditorUtilityCommon.h"
+#include "EditorUtilityObject.h"
 #include "EditorUtilitySubsystem.h"
 #include "EditorUtilityWidget.h"
 #include "EditorUtilityWidgetBlueprint.h"
@@ -87,10 +90,14 @@ public:
 
 		FKismetCompilerContext::RegisterCompilerForBP(UEditorUtilityWidgetBlueprint::StaticClass(), &UWidgetBlueprint::GetCompilerForWidgetBP);
 
-		// Register widget blueprint compiler we do this no matter what.
+		// Register widget blueprint compiler we do this no matter what - @todo: can i remove this? we're double registering..
 		IUMGEditorModule& UMGEditorModule = FModuleManager::LoadModuleChecked<IUMGEditorModule>("UMGEditor");
 		IKismetCompilerInterface& KismetCompilerModule = FModuleManager::LoadModuleChecked<IKismetCompilerInterface>("KismetCompiler");
 		KismetCompilerModule.GetCompilers().Add(UMGEditorModule.GetRegisteredCompiler());
+		KismetCompilerModule.OverrideBPTypeForClass(AEditorUtilityActor::StaticClass(), UEditorUtilityBlueprint::StaticClass());
+		KismetCompilerModule.OverrideBPTypeForClass(AEditorUtilityCamera::StaticClass(), UEditorUtilityBlueprint::StaticClass());
+		KismetCompilerModule.OverrideBPTypeForClass(UEditorUtilityObject::StaticClass(), UEditorUtilityBlueprint::StaticClass());
+		KismetCompilerModule.OverrideBPTypeForClassInEditor(UUserWidget::StaticClass(), UEditorUtilityWidgetBlueprint::StaticClass());
 
 		FBlutilityContentBrowserExtensions::InstallHooks();
 		FBlutilityLevelEditorExtensions::InstallHooks();
diff --git a/Engine/Source/Editor/KismetCompiler/Private/KismetCompiler.cpp b/Engine/Source/Editor/KismetCompiler/Private/KismetCompiler.cpp
--- a/Engine/Source/Editor/KismetCompiler/Private/KismetCompiler.cpp
+++ b/Engine/Source/Editor/KismetCompiler/Private/KismetCompiler.cpp
@@ -58,6 +58,7 @@
 #include "Engine/InheritableComponentHandler.h"
 #include "Serialization/ArchiveScriptReferenceCollector.h"
 #include "UObject/UnrealTypePrivate.h"
+#include "KismetCompilerModule.h"
 
 static bool bDebugPropertyPropagation = false;
 
@@ -4792,6 +4793,9 @@ void FKismetCompilerContext::CompileClassLayout(EInternalCompilerFlags InternalF
 		FBlueprintEditorUtils::ConformDelegateSignatureGraphs(Blueprint);
 	}
 
+	IKismetCompilerInterface& KismetCompilerModule = FModuleManager::LoadModuleChecked<IKismetCompilerInterface>("KismetCompiler");
+	KismetCompilerModule.ValidateBPAndClassType(Blueprint, MessageLog);
+
 	// If applicable, register any delegate proxy functions and their captured actor variables
 	RegisterClassDelegateProxiesFromBlueprint();
 
diff --git a/Engine/Source/Editor/KismetCompiler/Private/KismetCompilerModule.cpp b/Engine/Source/Editor/KismetCompiler/Private/KismetCompilerModule.cpp
--- a/Engine/Source/Editor/KismetCompiler/Private/KismetCompilerModule.cpp
+++ b/Engine/Source/Editor/KismetCompiler/Private/KismetCompilerModule.cpp
@@ -35,6 +35,8 @@ DECLARE_CYCLE_STAT(TEXT("Compile Generated Class"), EKismetCompilerStats_Compile
 
 #define LOCTEXT_NAMESPACE "KismetCompiler"
 
+#pragma optimize("", off)
+
 //////////////////////////////////////////////////////////////////////////
 // FKismet2CompilerModule - The Kismet 2 Compiler module
 class FKismet2CompilerModule : public IKismetCompilerInterface
@@ -46,9 +48,21 @@ public:
 	virtual void RecoverCorruptedBlueprint(class UBlueprint* Blueprint) override;
 	virtual void RemoveBlueprintGeneratedClasses(class UBlueprint* Blueprint) override;
 	virtual TArray<IBlueprintCompiler*>& GetCompilers() override { return Compilers; }
+	virtual void OverrideBPTypeForClass(UClass* Class, TSubclassOf<UBlueprint> BlueprintType) override;
+	virtual void OverrideBPTypeForClassInEditor(UClass* Class, TSubclassOf<UBlueprint> BlueprintType) override;
+	virtual void OverrideBPGCTypeForBPType(TSubclassOf<UBlueprint> BlueprintType, TSubclassOf<UBlueprintGeneratedClass> BPGCType) override;
+	virtual void ValidateBPAndClassType(UBlueprint* BP, FCompilerResultsLog& OutResults) override;
 	virtual void GetBlueprintTypesForClass(UClass* ParentClass, UClass*& OutBlueprintClass, UClass*& OutBlueprintGeneratedClass) const override;
 	// End implementation
+
+	static TSubclassOf<UBlueprint> FindBlueprintType(UClass* ForClass, const TMap<UClass*, TSubclassOf<UBlueprint>>& FromMap);
 private:
+	// these are all pointers to native reflection data, so don't require gc visibility
+	// this will frustrate hotreload, though - hot reload of objects used as keys or values
+	// doesn't really work anyway:
+	TMap<UClass*, TSubclassOf<UBlueprint>> ClassToBPType;
+	TMap<UClass*, TSubclassOf<UBlueprint>> ClassToEditorBPType;
+	TMap<TSubclassOf<UBlueprint>, TSubclassOf<UBlueprintGeneratedClass>> BPTypeToBPGCType;
 
 	TArray<IBlueprintCompiler*> Compilers;
 };
@@ -195,8 +209,168 @@ void FKismet2CompilerModule::RemoveBlueprintGeneratedClasses(class UBlueprint* B
 	}
 }
 
+void FKismet2CompilerModule::OverrideBPTypeForClass(UClass* Class, TSubclassOf<UBlueprint> BlueprintType)
+{
+	check(Class && BlueprintType);
+	#if DO_CHECK
+	if (const TSubclassOf<UBlueprint>* ExistingBlueprintType = ClassToBPType.Find(Class))
+	{
+		ensureMsgf(false,
+			TEXT("Ambiguous mapping attempting to add %s to %s when mapping to %s exists"),
+			*(Class->GetFullName()), *(BlueprintType.Get()->GetFullName()),
+			*(ExistingBlueprintType->Get()->GetFullName())
+		);
+	}
+	#endif // DO_CHECK
+
+	ClassToBPType.Add(Class, BlueprintType);
+}
+
+void FKismet2CompilerModule::OverrideBPTypeForClassInEditor(UClass* Class, TSubclassOf<UBlueprint> BlueprintType)
+{
+	check(Class && BlueprintType);
+#if DO_CHECK
+	if (const TSubclassOf<UBlueprint>* ExistingBlueprintType = ClassToEditorBPType.Find(Class))
+	{
+		ensureMsgf(false,
+			TEXT("Ambiguous mapping attempting to add %s to %s when mapping to %s exists"),
+			*(Class->GetFullName()), *(BlueprintType.Get()->GetFullName()),
+			*(ExistingBlueprintType->Get()->GetFullName())
+		);
+	}
+#endif // DO_CHECK
+
+	ClassToEditorBPType.Add(Class, BlueprintType);
+}
+
+void FKismet2CompilerModule::OverrideBPGCTypeForBPType(TSubclassOf<UBlueprint> BlueprintType, TSubclassOf<UBlueprintGeneratedClass> BPGCType)
+{
+	check(BlueprintType && BPGCType);
+#if DO_CHECK
+	if (const TSubclassOf<UBlueprintGeneratedClass>* ExistingBlueprintType = BPTypeToBPGCType.Find(BlueprintType))
+	{
+		ensureMsgf(false,
+			TEXT("Ambiguous mapping attempting to add %s to %s when mapping to %s exists"),
+			*(BlueprintType.Get()->GetFullName()), *(BPGCType.Get()->GetFullName()),
+			*(ExistingBlueprintType->Get()->GetFullName())
+		);
+	}
+#endif // DO_CHECK
+
+	BPTypeToBPGCType.Add(BlueprintType, BPGCType);
+}
+
+void FKismet2CompilerModule::ValidateBPAndClassType(UBlueprint* BP, FCompilerResultsLog& OutResults)
+{
+	if (BP->BlueprintType == BPTYPE_MacroLibrary)
+	{
+		// macros will contain macros of all sorts of UClasses, they only have a notional
+		// type for scoping purposes, it is not useful to inspect their GeneratedClass:
+		return;
+	}
+
+	UBlueprintGeneratedClass* BPGC = Cast<UBlueprintGeneratedClass>(BP->GeneratedClass);
+	if(!BPGC)
+	{
+		return;
+	}
+
+	// validate according to ClassToBPType:
+	{
+		TSubclassOf<UBlueprint> ExpectedType = FindBlueprintType(BPGC, ClassToBPType);
+		if(!BP->GetClass()->IsChildOf(ExpectedType))
+		{
+			OutResults.Warning(
+				*(FText::Format(
+					LOCTEXT("BPGCTypeMismatch", "@@ has an incorrect BP type - this type of blueprint ({0}) needs to be converted."),
+					FText::FromString(BP->GetFullName())).ToString()),
+				BP
+			);
+		}
+	}
+
+	// validate according to ClassToEditorBPType:
+	if (::IsEditorOnlyObject(BP->GetOutermost()))
+	{
+		TSubclassOf<UBlueprint> ExpectedType = FindBlueprintType(BPGC, ClassToEditorBPType);
+		if (!BP->GetClass()->IsChildOf(ExpectedType))
+		{
+			OutResults.Warning(
+				*(FText::Format(
+					LOCTEXT("BPGCTypeMismatch", "@@ has an incorrect editor BP type - this type of blueprint ({0}) needs to be converted."),
+					FText::FromString(BP->GetFullName())).ToString()),
+				BP
+			);
+		}
+	}
+
+	// validate according to BPTypeToBPGCType:
+	for (TPair<TSubclassOf<UBlueprint>, TSubclassOf<UBlueprintGeneratedClass>> BPTypeToBPGCTypeIter : BPTypeToBPGCType)
+	{
+		// if there's a BP type implied by the BPGC we want it to match the provided BP:
+		if(BP->GetClass()->IsChildOf(BPTypeToBPGCTypeIter.Key))
+		{
+			if (!BPGC->GetClass()->IsChildOf(BPTypeToBPGCTypeIter.Value))
+			{
+				OutResults.Warning(
+					*(FText::Format(
+						LOCTEXT("BPGCTypeMismatch", "@@ has an incorrect BPGC type - this type of blueprint ({0}) needs to sanitize its class."),
+						FText::FromString(BP->GetFullName())).ToString()),
+					BP
+				);
+			}
+		}
+
+		// if there's a class implied by the BP type we want it to match the provided BPGC:
+		if(BPGC->GetClass()->IsChildOf(BPTypeToBPGCTypeIter.Value))
+		{
+			if (!BP->GetClass()->IsChildOf(BPTypeToBPGCTypeIter.Key))
+			{
+				OutResults.Warning(
+					*(FText::Format(
+						LOCTEXT("BPGCTypeMismatch", "@@ has an incorrect BP type - this blueprint ({0}) needs to be converted to a different type of blueprint."),
+						FText::FromString(BP->GetFullName())).ToString()),
+					BP
+				);
+			}
+		}
+	}
+}
+
 void FKismet2CompilerModule::GetBlueprintTypesForClass(UClass* ParentClass, UClass*& OutBlueprintClass, UClass*& OutBlueprintGeneratedClass) const
 {
+	const TMap<TSubclassOf<UBlueprint>, TSubclassOf<UBlueprintGeneratedClass>>& BPTypeToBPGCTypeForLookup = BPTypeToBPGCType;
+	const auto GetBPGCTypeForBPType = [&BPTypeToBPGCTypeForLookup](TSubclassOf<UBlueprint> BPType) -> TSubclassOf<UBlueprintGeneratedClass>
+	{
+		if(const TSubclassOf<UBlueprintGeneratedClass>* BPGCType = BPTypeToBPGCTypeForLookup.Find(BPType))
+		{
+			return *BPGCType;
+		}
+		return UBlueprintGeneratedClass::StaticClass();
+	};
+
+	// honor ClassToEditorBPType
+	if(::IsEditorOnlyObject(ParentClass))
+	{
+		TSubclassOf<UBlueprint> BlueprintType = FindBlueprintType(ParentClass, ClassToEditorBPType);
+		if(BlueprintType != UBlueprint::StaticClass())
+		{
+			OutBlueprintGeneratedClass = GetBPGCTypeForBPType(BlueprintType);
+			OutBlueprintClass = BlueprintType;
+			return;
+		}
+	}
+
+	// no editor mapping found, fall back to ClassToBPType:
+	TSubclassOf<UBlueprint> BlueprintType = FindBlueprintType(ParentClass, ClassToBPType);
+	if (BlueprintType != UBlueprint::StaticClass())
+	{
+		OutBlueprintGeneratedClass = GetBPGCTypeForBPType(BlueprintType);
+		OutBlueprintClass = BlueprintType;
+		return;
+	}
+
+	// legacy support for IBlueprintCompiler:
 	for ( IBlueprintCompiler* Compiler : Compilers )
 	{
 		if ( Compiler->GetBlueprintTypesForClass(ParentClass, OutBlueprintClass, OutBlueprintGeneratedClass) )
@@ -209,4 +383,18 @@ void FKismet2CompilerModule::GetBlueprintTypesForClass(UClass* ParentClass, UCla
 	OutBlueprintGeneratedClass = UBlueprintGeneratedClass::StaticClass();
 }
 
+TSubclassOf<UBlueprint> FKismet2CompilerModule::FindBlueprintType(UClass* ForClass, const TMap<UClass*, TSubclassOf<UBlueprint>>& FromMap)
+{
+	UClass* Iter = ForClass;
+	while (Iter)
+	{
+		if (const TSubclassOf<UBlueprint>* BPType = FromMap.Find(Iter))
+		{
+			return *BPType;
+		}
+		Iter = Iter->GetSuperClass();
+	}
+	return UBlueprint::StaticClass();
+}
+
 #undef LOCTEXT_NAMESPACE
diff --git a/Engine/Source/Editor/KismetCompiler/Public/KismetCompilerModule.h b/Engine/Source/Editor/KismetCompiler/Public/KismetCompilerModule.h
--- a/Engine/Source/Editor/KismetCompiler/Public/KismetCompilerModule.h
+++ b/Engine/Source/Editor/KismetCompiler/Public/KismetCompilerModule.h
@@ -94,6 +94,13 @@ public:
 	 */
 	virtual TArray<IBlueprintCompiler*>& GetCompilers() = 0;
 
+	/** Facilities for establishing mappings between UClasses, UBlueprints, and UBlueprintGenerated Classes*/
+	virtual void OverrideBPTypeForClass(UClass* Class, TSubclassOf<UBlueprint> BlueprintType) = 0;
+	virtual void OverrideBPTypeForClassInEditor(UClass* Class, TSubclassOf<UBlueprint> BlueprintType) = 0;
+	virtual void OverrideBPGCTypeForBPType(TSubclassOf<UBlueprint> BlueprintType, TSubclassOf<UBlueprintGeneratedClass> BPGCType) = 0;
+
+	virtual void ValidateBPAndClassType(UBlueprint* BP, FCompilerResultsLog& OutResults) = 0;
+
 	/**
 	 * Get the blueprint class and generated blueprint class for a particular class type.  Not every
 	 * blueprint is a normal UBlueprint, like UUserWidget blueprints should be UWidgetBlueprints.
diff --git a/Engine/Source/Editor/UMGEditor/Private/UMGEditorModule.cpp b/Engine/Source/Editor/UMGEditor/Private/UMGEditorModule.cpp
--- a/Engine/Source/Editor/UMGEditor/Private/UMGEditorModule.cpp
+++ b/Engine/Source/Editor/UMGEditor/Private/UMGEditorModule.cpp
@@ -87,6 +87,7 @@ public:
 		// Register widget blueprint compiler we do this no matter what.
 		IKismetCompilerInterface& KismetCompilerModule = FModuleManager::LoadModuleChecked<IKismetCompilerInterface>("KismetCompiler");
 		KismetCompilerModule.GetCompilers().Add(&WidgetBlueprintCompiler);
+		KismetCompilerModule.OverrideBPTypeForClass(UUserWidget::StaticClass(), UWidgetBlueprint::StaticClass());
 
 		// Add Customization for variable in Graph editor
 		if (FBlueprintEditorModule* BlueprintEditorModule = FModuleManager::GetModulePtr<FBlueprintEditorModule>("Kismet"))
