Remove optimization flag
#rnx

diff --git a/Engine/Source/Editor/KismetCompiler/Private/KismetCompilerModule.cpp b/Engine/Source/Editor/KismetCompiler/Private/KismetCompilerModule.cpp
--- a/Engine/Source/Editor/KismetCompiler/Private/KismetCompilerModule.cpp
+++ b/Engine/Source/Editor/KismetCompiler/Private/KismetCompilerModule.cpp
@@ -35,8 +35,6 @@ DECLARE_CYCLE_STAT(TEXT("Compile Generated Class"), EKismetCompilerStats_Compile
 
 #define LOCTEXT_NAMESPACE "KismetCompiler"
 
-#pragma optimize("", off)
-
 //////////////////////////////////////////////////////////////////////////
 // FKismet2CompilerModule - The Kismet 2 Compiler module
 class FKismet2CompilerModule : public IKismetCompilerInterface
