Non unity fix
#rnx

diff --git a/Engine/Source/Editor/KismetCompiler/Public/KismetCompilerModule.h b/Engine/Source/Editor/KismetCompiler/Public/KismetCompilerModule.h
--- a/Engine/Source/Editor/KismetCompiler/Public/KismetCompilerModule.h
+++ b/Engine/Source/Editor/KismetCompiler/Public/KismetCompilerModule.h
@@ -7,6 +7,7 @@
 #include "CoreMinimal.h"
 #include "CoreTypes.h"
 #include "Modules/ModuleInterface.h"
+#include "Templates/SubclassOf.h"
 
 class FBlueprintCompileReinstancer;
 class FCompilerResultsLog;
