Downgrade this warning to note
#rnx

diff --git a/Engine/Source/Editor/KismetCompiler/Private/KismetCompilerModule.cpp b/Engine/Source/Editor/KismetCompiler/Private/KismetCompilerModule.cpp
--- a/Engine/Source/Editor/KismetCompiler/Private/KismetCompilerModule.cpp
+++ b/Engine/Source/Editor/KismetCompiler/Private/KismetCompilerModule.cpp
@@ -260,6 +260,8 @@ void FKismet2CompilerModule::OverrideBPGCTypeForBPType(TSubclassOf<UBlueprint> B
 
 void FKismet2CompilerModule::ValidateBPAndClassType(UBlueprint* BP, FCompilerResultsLog& OutResults)
 {
+	// validation can become a warning as this matures, for now just note:
+
 	if (BP->BlueprintType == BPTYPE_MacroLibrary)
 	{
 		// macros will contain macros of all sorts of UClasses, they only have a notional
@@ -278,7 +280,7 @@ void FKismet2CompilerModule::ValidateBPAndClassType(UBlueprint* BP, FCompilerRes
 		TSubclassOf<UBlueprint> ExpectedType = FindBlueprintType(BPGC, ClassToBPType);
 		if(!BP->GetClass()->IsChildOf(ExpectedType))
 		{
-			OutResults.Warning(
+			OutResults.Note(
 				*(FText::Format(
 					LOCTEXT("BPGCTypeMismatch", "@@ has an incorrect BP type - this type of blueprint ({0}) needs to be converted."),
 					FText::FromString(BP->GetFullName())).ToString()),
@@ -293,7 +295,7 @@ void FKismet2CompilerModule::ValidateBPAndClassType(UBlueprint* BP, FCompilerRes
 		TSubclassOf<UBlueprint> ExpectedType = FindBlueprintType(BPGC, ClassToEditorBPType);
 		if (!BP->GetClass()->IsChildOf(ExpectedType))
 		{
-			OutResults.Warning(
+			OutResults.Note(
 				*(FText::Format(
 					LOCTEXT("BPGCTypeMismatch", "@@ has an incorrect editor BP type - this type of blueprint ({0}) needs to be converted."),
 					FText::FromString(BP->GetFullName())).ToString()),
@@ -310,7 +312,7 @@ void FKismet2CompilerModule::ValidateBPAndClassType(UBlueprint* BP, FCompilerRes
 		{
 			if (!BPGC->GetClass()->IsChildOf(BPTypeToBPGCTypeIter.Value))
 			{
-				OutResults.Warning(
+				OutResults.Note(
 					*(FText::Format(
 						LOCTEXT("BPGCTypeMismatch", "@@ has an incorrect BPGC type - this type of blueprint ({0}) needs to sanitize its class."),
 						FText::FromString(BP->GetFullName())).ToString()),
@@ -324,7 +326,7 @@ void FKismet2CompilerModule::ValidateBPAndClassType(UBlueprint* BP, FCompilerRes
 		{
 			if (!BP->GetClass()->IsChildOf(BPTypeToBPGCTypeIter.Key))
 			{
-				OutResults.Warning(
+				OutResults.Note(
 					*(FText::Format(
 						LOCTEXT("BPGCTypeMismatch", "@@ has an incorrect BP type - this blueprint ({0}) needs to be converted to a different type of blueprint."),
 						FText::FromString(BP->GetFullName())).ToString()),
