[CIS] Fix LOCTEXT key collision.

#rnx

diff --git a/Engine/Source/Editor/KismetCompiler/Private/KismetCompilerModule.cpp b/Engine/Source/Editor/KismetCompiler/Private/KismetCompilerModule.cpp
--- a/Engine/Source/Editor/KismetCompiler/Private/KismetCompilerModule.cpp
+++ b/Engine/Source/Editor/KismetCompiler/Private/KismetCompilerModule.cpp
@@ -282,7 +282,7 @@ void FKismet2CompilerModule::ValidateBPAndClassType(UBlueprint* BP, FCompilerRes
 		{
 			OutResults.Note(
 				*(FText::Format(
-					LOCTEXT("BPGCTypeMismatch", "@@ has an incorrect BP type - this type of blueprint ({0}) needs to be converted."),
+					LOCTEXT("BPGCTypeMismatch_Blueprint", "@@ has an incorrect BP type - this type of blueprint ({0}) needs to be converted."),
 					FText::FromString(BP->GetFullName())).ToString()),
 				BP
 			);
@@ -297,7 +297,7 @@ void FKismet2CompilerModule::ValidateBPAndClassType(UBlueprint* BP, FCompilerRes
 		{
 			OutResults.Note(
 				*(FText::Format(
-					LOCTEXT("BPGCTypeMismatch", "@@ has an incorrect editor BP type - this type of blueprint ({0}) needs to be converted."),
+					LOCTEXT("BPGCTypeMismatch_EditorOnlyBlueprint", "@@ has an incorrect editor BP type - this type of blueprint ({0}) needs to be converted."),
 					FText::FromString(BP->GetFullName())).ToString()),
 				BP
 			);
@@ -314,7 +314,7 @@ void FKismet2CompilerModule::ValidateBPAndClassType(UBlueprint* BP, FCompilerRes
 			{
 				OutResults.Note(
 					*(FText::Format(
-						LOCTEXT("BPGCTypeMismatch", "@@ has an incorrect BPGC type - this type of blueprint ({0}) needs to sanitize its class."),
+						LOCTEXT("BPGCTypeMismatch_GeneratedClass", "@@ has an incorrect BPGC type - this type of blueprint ({0}) needs to sanitize its class."),
 						FText::FromString(BP->GetFullName())).ToString()),
 					BP
 				);
@@ -328,7 +328,7 @@ void FKismet2CompilerModule::ValidateBPAndClassType(UBlueprint* BP, FCompilerRes
 			{
 				OutResults.Note(
 					*(FText::Format(
-						LOCTEXT("BPGCTypeMismatch", "@@ has an incorrect BP type - this blueprint ({0}) needs to be converted to a different type of blueprint."),
+						LOCTEXT("BPGCTypeMismatch_ClassGeneratedBy", "@@ has an incorrect BP type - this blueprint ({0}) needs to be converted to a different type of blueprint."),
 						FText::FromString(BP->GetFullName())).ToString()),
 					BP
 				);
