Destroy the skeleton class before clearing the UBlueprint's generated class. This allows UBlueprintGeneratedClass::PostLoad (called during duplication) to correctly identify that it is not the authoritative class, and therefore not generate GC info

diff --git a/Engine/Source/Editor/KismetCompiler/Private/KismetCompilerModule.cpp b/Engine/Source/Editor/KismetCompiler/Private/KismetCompilerModule.cpp
--- a/Engine/Source/Editor/KismetCompiler/Private/KismetCompilerModule.cpp
+++ b/Engine/Source/Editor/KismetCompiler/Private/KismetCompilerModule.cpp
@@ -193,17 +193,19 @@ void FKismet2CompilerModule::RemoveBlueprintGeneratedClasses(class UBlueprint* B
 {
 	if (Blueprint != NULL)
 	{
-		if (Blueprint->GeneratedClass != NULL)
-		{
-			FKismetCompilerUtilities::ConsignToOblivion(Blueprint->GeneratedClass, Blueprint->bIsRegeneratingOnLoad);
-			Blueprint->GeneratedClass = NULL;
-		}
-
+		// Order unfortunately matters, as we want to allow UBlueprintGeneratedClass::GetAuthoritativeClass to function
+		// correctly for as long as possible:
 		if (Blueprint->SkeletonGeneratedClass != NULL)
 		{
 			FKismetCompilerUtilities::ConsignToOblivion(Blueprint->SkeletonGeneratedClass, Blueprint->bIsRegeneratingOnLoad);
 			Blueprint->SkeletonGeneratedClass = NULL;
 		}
+
+		if (Blueprint->GeneratedClass != NULL)
+		{
+			FKismetCompilerUtilities::ConsignToOblivion(Blueprint->GeneratedClass, Blueprint->bIsRegeneratingOnLoad);
+			Blueprint->GeneratedClass = NULL;
+		}
 	}
 }
 
