Avoid potential collisions with existing non-transient objects when creating transient event stub graph artifacts during Blueprint ubergraph compilation.

diff --git a/Engine/Source/Editor/KismetCompiler/Private/KismetCompiler.cpp b/Engine/Source/Editor/KismetCompiler/Private/KismetCompiler.cpp
--- a/Engine/Source/Editor/KismetCompiler/Private/KismetCompiler.cpp
+++ b/Engine/Source/Editor/KismetCompiler/Private/KismetCompiler.cpp
@@ -3552,20 +3552,36 @@ void FKismetCompilerContext::CreateFunctionStubForEvent(UK2Node_Event* SrcEventN
 
 	// Create the stub graph and add it to the list of functions to compile
 
-	UObject* ExistingGraph = static_cast<UObject*>(FindObjectWithOuter(OwnerOfTemporaries, UEdGraph::StaticClass(), EventNodeName));
-	if (ExistingGraph && !ExistingGraph->HasAnyFlags(RF_Transient))
+	FName EventNodeObjectName = EventNodeName;
+	const UObject* ExistingObject = static_cast<UObject*>(FindObjectWithOuter(OwnerOfTemporaries, nullptr, EventNodeName));
+	if (ExistingObject)
 	{
-		MessageLog.Error(
-			*FText::Format(
-				LOCTEXT("CannotCreateStubForEvent_ErrorFmt", "Graph named '{0}' already exists in '{1}'. Another one cannot be generated from @@"),
-				FText::FromName(EventNodeName),
-				FText::FromString(*GetNameSafe(OwnerOfTemporaries))
-			).ToString(),
-			SrcEventNode
-		);
-		return;
+		if (ExistingObject->IsA<UEdGraph>() && !ExistingObject->HasAnyFlags(RF_Transient))
+		{
+			MessageLog.Error(
+				*FText::Format(
+					LOCTEXT("CannotCreateStubForEvent_ErrorFmt", "Graph named '{0}' already exists in '{1}'. Another one cannot be generated from @@"),
+					FText::FromName(EventNodeName),
+					FText::FromString(*GetNameSafe(OwnerOfTemporaries))
+				).ToString(),
+				SrcEventNode
+			);
+			return;
+		}
+		else
+		{
+			// If we collide with another non-transient object type or any child object owned by the same object that will also contain transient
+			// intermediate build artifacts post-compile, choose a unique name for the stub node. It's possible that the compile can still fail
+			// due to collisions between generated properties and/or functions within the compiled class object, but choosing a unique name here
+			// will avoid any issues with allocating the temporary node object below for the intermediate graph, so we at least get to that point.
+			// 
+			// @todo - Consider using a transient owner for build artifacts to avoid collisions with non-transient siblings that aren't going to
+			// otherwise get cleared out on compile.
+			EventNodeObjectName = MakeUniqueObjectName(OwnerOfTemporaries, UEdGraph::StaticClass(), EventNodeName);
+		}
 	}
-	UEdGraph* ChildStubGraph = NewObject<UEdGraph>(OwnerOfTemporaries, EventNodeName);
+
+	UEdGraph* ChildStubGraph = NewObject<UEdGraph>(OwnerOfTemporaries, EventNodeObjectName);
 	Blueprint->EventGraphs.Add(ChildStubGraph);
 	ChildStubGraph->Schema = UEdGraphSchema_K2::StaticClass();
 	ChildStubGraph->SetFlags(RF_Transient);
