[BugFix] Made it so users can no longer reparent blueprints to classes managed by a different blueprint type (IE Actor to EditorUtilityActor)

diff --git a/Engine/Source/Editor/KismetCompiler/Private/KismetCompilerModule.cpp b/Engine/Source/Editor/KismetCompiler/Private/KismetCompilerModule.cpp
--- a/Engine/Source/Editor/KismetCompiler/Private/KismetCompilerModule.cpp
+++ b/Engine/Source/Editor/KismetCompiler/Private/KismetCompilerModule.cpp
@@ -51,6 +51,7 @@ public:
 	virtual void OverrideBPGCTypeForBPType(TSubclassOf<UBlueprint> BlueprintType, TSubclassOf<UBlueprintGeneratedClass> BPGCType) override;
 	virtual void ValidateBPAndClassType(UBlueprint* BP, FCompilerResultsLog& OutResults) override;
 	virtual void GetBlueprintTypesForClass(UClass* ParentClass, UClass*& OutBlueprintClass, UClass*& OutBlueprintGeneratedClass) const override;
+	virtual void GetSubclassesWithDifferingBlueprintTypes(UClass* Class, TSet<const UClass*>& OutMismatchedSubclasses) const override;
 	// End implementation
 
 	static TSubclassOf<UBlueprint> FindBlueprintType(UClass* ForClass, const TMap<UClass*, TSubclassOf<UBlueprint>>& FromMap);
@@ -385,6 +386,26 @@ void FKismet2CompilerModule::GetBlueprintTypesForClass(UClass* ParentClass, UCla
 	OutBlueprintGeneratedClass = UBlueprintGeneratedClass::StaticClass();
 }
 
+void FKismet2CompilerModule::GetSubclassesWithDifferingBlueprintTypes(UClass* Class, TSet<const UClass*>& OutMismatchedSubclasses) const
+{
+	UClass* BPClass;
+	UClass* BPGeneratedClass;
+	GetBlueprintTypesForClass(Class, BPClass, BPGeneratedClass);
+	
+	auto CheckClassToBPTypeMap = [](const TMap<UClass*, TSubclassOf<UBlueprint>>& Map, const UClass* Class, const UClass* BPClass, TSet<const UClass*>& Result)
+	{
+		for (const auto&[SupportedClass, OtherBPClass] : Map)
+		{
+			if (OtherBPClass != BPClass && SupportedClass->IsChildOf(Class))
+			{
+				Result.Add(SupportedClass);
+			}
+		}
+	};
+	CheckClassToBPTypeMap(ClassToBPType, Class, BPClass, OutMismatchedSubclasses);
+	CheckClassToBPTypeMap(ClassToEditorBPType, Class, BPClass, OutMismatchedSubclasses);
+}
+
 TSubclassOf<UBlueprint> FKismet2CompilerModule::FindBlueprintType(UClass* ForClass, const TMap<UClass*, TSubclassOf<UBlueprint>>& FromMap)
 {
 	UClass* Iter = ForClass;
diff --git a/Engine/Source/Editor/KismetCompiler/Public/KismetCompilerModule.h b/Engine/Source/Editor/KismetCompiler/Public/KismetCompilerModule.h
--- a/Engine/Source/Editor/KismetCompiler/Public/KismetCompilerModule.h
+++ b/Engine/Source/Editor/KismetCompiler/Public/KismetCompilerModule.h
@@ -107,6 +107,7 @@ public:
 	 * blueprint is a normal UBlueprint, like UUserWidget blueprints should be UWidgetBlueprints.
 	 */
 	virtual void GetBlueprintTypesForClass(UClass* ParentClass, UClass*& OutBlueprintClass, UClass*& OutBlueprintGeneratedClass) const = 0;
+	virtual void GetSubclassesWithDifferingBlueprintTypes(UClass* Class, TSet<const UClass*>& OutMismatchedSubclasses) const = 0;
 
 	// @todo: BP2CPP_remove
 PRAGMA_DISABLE_DEPRECATION_WARNINGS
diff --git a/Engine/Source/Editor/UnrealEd/Private/Kismet2/BlueprintEditorUtils.cpp b/Engine/Source/Editor/UnrealEd/Private/Kismet2/BlueprintEditorUtils.cpp
--- a/Engine/Source/Editor/UnrealEd/Private/Kismet2/BlueprintEditorUtils.cpp
+++ b/Engine/Source/Editor/UnrealEd/Private/Kismet2/BlueprintEditorUtils.cpp
@@ -8574,6 +8574,12 @@ TSharedRef<SWidget> FBlueprintEditorUtils::ConstructBlueprintParentClassPicker(
 		{
 			Filter->DisallowedClasses.Add(Blueprint->GeneratedClass);
 		}
+
+		// Disallow reparenting to any sub-class that has a different blueprint type.
+		// for example, reparenting AActor to AEditorUtilityActor is illegal because AActor uses UBlueprint while
+		// AEditorUtilityActor uses UEditorUtilityBlueprint
+		const IKismetCompilerInterface& KismetCompilerModule = FModuleManager::LoadModuleChecked<IKismetCompilerInterface>("KismetCompiler");
+		KismetCompilerModule.GetSubclassesWithDifferingBlueprintTypes(Blueprint->ParentClass, Filter->DisallowedChildrenOfClasses);
 	}
 
 	return FModuleManager::LoadModuleChecked<FClassViewerModule>("ClassViewer").CreateClassViewer(Options, OnPicked);
