Kismet Compiler: Deprecate SpawnIntermediateEventNode
#rnx

diff --git a/Engine/Plugins/EnhancedInput/Source/InputBlueprintNodes/Private/K2Node_EnhancedInputAction.cpp b/Engine/Plugins/EnhancedInput/Source/InputBlueprintNodes/Private/K2Node_EnhancedInputAction.cpp
--- a/Engine/Plugins/EnhancedInput/Source/InputBlueprintNodes/Private/K2Node_EnhancedInputAction.cpp
+++ b/Engine/Plugins/EnhancedInput/Source/InputBlueprintNodes/Private/K2Node_EnhancedInputAction.cpp
@@ -323,7 +323,7 @@ void UK2Node_EnhancedInputAction::ExpandNode(FKismetCompilerContext& CompilerCon
 			return nullptr;
 		}
 
-		UK2Node_EnhancedInputActionEvent* InputActionEvent = CompilerContext.SpawnIntermediateEventNode<UK2Node_EnhancedInputActionEvent>(this, Pin, SourceGraph);
+		UK2Node_EnhancedInputActionEvent* InputActionEvent = CompilerContext.SpawnIntermediateNode<UK2Node_EnhancedInputActionEvent>(this, SourceGraph);
 		InputActionEvent->CustomFunctionName = FName(*FString::Printf(TEXT("InpActEvt_%s_%s"), *GetActionName().ToString(), *InputActionEvent->GetName()));
 		InputActionEvent->InputAction = InputAction;
 		InputActionEvent->TriggerEvent = TriggerEvent;
diff --git a/Engine/Plugins/EnhancedInput/Source/InputBlueprintNodes/Private/K2Node_InputDebugKey.cpp b/Engine/Plugins/EnhancedInput/Source/InputBlueprintNodes/Private/K2Node_InputDebugKey.cpp
--- a/Engine/Plugins/EnhancedInput/Source/InputBlueprintNodes/Private/K2Node_InputDebugKey.cpp
+++ b/Engine/Plugins/EnhancedInput/Source/InputBlueprintNodes/Private/K2Node_InputDebugKey.cpp
@@ -263,7 +263,7 @@ void UK2Node_InputDebugKey::ExpandNode(FKismetCompilerContext& CompilerContext,
 	auto CreateDebugKeyEvent = [this, &CompilerContext, &SourceGraph](UEdGraphPin* Pin, EInputEvent EventType)
 	{
 		// Create the input touch event
-		UK2Node_InputDebugKeyEvent* InputDebugKeyEvent = CompilerContext.SpawnIntermediateEventNode<UK2Node_InputDebugKeyEvent>(this, Pin, SourceGraph);
+		UK2Node_InputDebugKeyEvent* InputDebugKeyEvent = CompilerContext.SpawnIntermediateNode<UK2Node_InputDebugKeyEvent>(this, SourceGraph);
 		const FName ModifierName = GetModifierName();
 		if (ModifierName != NAME_None)
 		{
diff --git a/Engine/Plugins/Experimental/Web/WebAPI/Source/WebAPIBlueprintGraph/Private/K2Node_WebAPIOperation.cpp b/Engine/Plugins/Experimental/Web/WebAPI/Source/WebAPIBlueprintGraph/Private/K2Node_WebAPIOperation.cpp
--- a/Engine/Plugins/Experimental/Web/WebAPI/Source/WebAPIBlueprintGraph/Private/K2Node_WebAPIOperation.cpp
+++ b/Engine/Plugins/Experimental/Web/WebAPI/Source/WebAPIBlueprintGraph/Private/K2Node_WebAPIOperation.cpp
@@ -100,7 +100,7 @@ namespace UE::WebAPI::Private
 			return false;
 		}
 
-		UK2Node_CustomEvent* CurrentCustomEventNode = CompilerContext.SpawnIntermediateEventNode<UK2Node_CustomEvent>(CurrentNode, PinForCurrentDelegateProperty, SourceGraph);
+		UK2Node_CustomEvent* CurrentCustomEventNode = CompilerContext.SpawnIntermediateNode<UK2Node_CustomEvent>(CurrentNode, SourceGraph);
 		{
 		UK2Node_AddDelegate* AddDelegateNode = CompilerContext.SpawnIntermediateNode<UK2Node_AddDelegate>(CurrentNode, SourceGraph);
 		AddDelegateNode->SetFromProperty(CurrentProperty, false, CurrentProperty->GetOwnerClass());
diff --git a/Engine/Source/Editor/AnimGraph/Private/AnimBlueprintExtension_Base.cpp b/Engine/Source/Editor/AnimGraph/Private/AnimBlueprintExtension_Base.cpp
--- a/Engine/Source/Editor/AnimGraph/Private/AnimBlueprintExtension_Base.cpp
+++ b/Engine/Source/Editor/AnimGraph/Private/AnimBlueprintExtension_Base.cpp
@@ -549,7 +549,7 @@ void UAnimBlueprintExtension_Base::CreateEvaluationHandler(IAnimBlueprintCompila
 	HandlerFunctionNames.Add(Record.HandlerFunctionName);
 
 	// Add a custom event in the graph
-	UK2Node_CustomEvent* CustomEventNode = InCompilationContext.SpawnIntermediateEventNode<UK2Node_CustomEvent>(InNode, nullptr, InCompilationContext.GetConsolidatedEventGraph());
+	UK2Node_CustomEvent* CustomEventNode = InCompilationContext.SpawnIntermediateNode<UK2Node_CustomEvent>(InNode, InCompilationContext.GetConsolidatedEventGraph());
 	CustomEventNode->bInternalEvent = true;
 	CustomEventNode->CustomFunctionName = Record.HandlerFunctionName;
 	CustomEventNode->AllocateDefaultPins();
diff --git a/Engine/Source/Editor/AnimGraph/Private/AnimGraphNode_CallFunction.cpp b/Engine/Source/Editor/AnimGraph/Private/AnimGraphNode_CallFunction.cpp
--- a/Engine/Source/Editor/AnimGraph/Private/AnimGraphNode_CallFunction.cpp
+++ b/Engine/Source/Editor/AnimGraph/Private/AnimGraphNode_CallFunction.cpp
@@ -131,7 +131,7 @@ void UAnimGraphNode_CallFunction::ExpandNode(FKismetCompilerContext& CompilerCon
 
 		if(EventName != NAME_None)
 		{
-			UK2Node_CustomEvent* CustomEventNode = CompilerContext.SpawnIntermediateEventNode<UK2Node_CustomEvent>(this, nullptr, CompilerContext.ConsolidatedEventGraph);
+			UK2Node_CustomEvent* CustomEventNode = CompilerContext.SpawnIntermediateNode<UK2Node_CustomEvent>(this, CompilerContext.ConsolidatedEventGraph);
 			CustomEventNode->bInternalEvent = true;
 			CustomEventNode->CustomFunctionName = EventName;
 			CustomEventNode->AllocateDefaultPins();
@@ -139,7 +139,7 @@ void UAnimGraphNode_CallFunction::ExpandNode(FKismetCompilerContext& CompilerCon
 			UEdGraphPin* ExecChain = K2Schema->FindExecutionPin(*CustomEventNode, EGPD_Output);
 
 			// Add call function node
-			UK2Node_CallFunction* NewCallFunctionNode = CompilerContext.SpawnIntermediateEventNode<UK2Node_CallFunction>(this, nullptr, CompilerContext.ConsolidatedEventGraph);
+			UK2Node_CallFunction* NewCallFunctionNode = CompilerContext.SpawnIntermediateNode<UK2Node_CallFunction>(this, CompilerContext.ConsolidatedEventGraph);
 			NewCallFunctionNode->FunctionReference = CallFunctionPrototype->FunctionReference;
 			NewCallFunctionNode->AllocateDefaultPins();
 			
diff --git a/Engine/Source/Editor/AnimGraph/Public/IAnimBlueprintCompilationContext.h b/Engine/Source/Editor/AnimGraph/Public/IAnimBlueprintCompilationContext.h
--- a/Engine/Source/Editor/AnimGraph/Public/IAnimBlueprintCompilationContext.h
+++ b/Engine/Source/Editor/AnimGraph/Public/IAnimBlueprintCompilationContext.h
@@ -74,9 +74,10 @@ public:
 
 	// Spawns an intermediate event node associated with the source node (for error purposes)
 	template <typename NodeType>
+	UE_DEPRECATED(5.4, "SpawnIntermediateEventNode is equivalent to SpawnIntermediateNode, this redundant function has been deprecated.")
 	NodeType* SpawnIntermediateEventNode(UEdGraphNode* SourceNode, UEdGraphPin* SourcePin = nullptr, UEdGraph* ParentGraph = nullptr)
 	{
-		return GetKismetCompiler()->SpawnIntermediateEventNode<NodeType>(SourceNode, SourcePin, ParentGraph);
+		return GetKismetCompiler()->SpawnIntermediateNode<NodeType>(SourceNode, ParentGraph);
 	}
 
 	// Find a property in the currently-compiled class
diff --git a/Engine/Source/Editor/BlueprintGraph/Private/K2Node_BaseAsyncTask.cpp b/Engine/Source/Editor/BlueprintGraph/Private/K2Node_BaseAsyncTask.cpp
--- a/Engine/Source/Editor/BlueprintGraph/Private/K2Node_BaseAsyncTask.cpp
+++ b/Engine/Source/Editor/BlueprintGraph/Private/K2Node_BaseAsyncTask.cpp
@@ -311,11 +311,11 @@ bool UK2Node_BaseAsyncTask::FBaseAsyncTaskHelper::HandleDelegateImplementation(
 		return false;
 	}
 
-	UK2Node_CustomEvent* CurrentCENode = CompilerContext.SpawnIntermediateEventNode<UK2Node_CustomEvent>(CurrentNode, PinForCurrentDelegateProperty, SourceGraph);
+	UK2Node_CustomEvent* CurrentCENode = CompilerContext.SpawnIntermediateNode<UK2Node_CustomEvent>(CurrentNode, SourceGraph);
 	{
-	UK2Node_AddDelegate* AddDelegateNode = CompilerContext.SpawnIntermediateNode<UK2Node_AddDelegate>(CurrentNode, SourceGraph);
-	AddDelegateNode->SetFromProperty(CurrentProperty, false, CurrentProperty->GetOwnerClass());
-	AddDelegateNode->AllocateDefaultPins();
+		UK2Node_AddDelegate* AddDelegateNode = CompilerContext.SpawnIntermediateNode<UK2Node_AddDelegate>(CurrentNode, SourceGraph);
+		AddDelegateNode->SetFromProperty(CurrentProperty, false, CurrentProperty->GetOwnerClass());
+		AddDelegateNode->AllocateDefaultPins();
 		bIsErrorFree &= Schema->TryCreateConnection(AddDelegateNode->FindPinChecked(UEdGraphSchema_K2::PN_Self), ProxyObjectPin);
 		bIsErrorFree &= Schema->TryCreateConnection(InOutLastThenPin, AddDelegateNode->FindPinChecked(UEdGraphSchema_K2::PN_Execute));
 		InOutLastThenPin = AddDelegateNode->FindPinChecked(UEdGraphSchema_K2::PN_Then);
diff --git a/Engine/Source/Editor/BlueprintGraph/Private/K2Node_DelegateSet.cpp b/Engine/Source/Editor/BlueprintGraph/Private/K2Node_DelegateSet.cpp
--- a/Engine/Source/Editor/BlueprintGraph/Private/K2Node_DelegateSet.cpp
+++ b/Engine/Source/Editor/BlueprintGraph/Private/K2Node_DelegateSet.cpp
@@ -285,7 +285,7 @@ void UK2Node_DelegateSet::ExpandNode(class FKismetCompilerContext& CompilerConte
 		if (UFunction* TargetFunction = GetDelegateSignature())
 		{
 			// First, create an event node matching the delegate signature
-			UK2Node_Event* DelegateEvent = CompilerContext.SpawnIntermediateEventNode<UK2Node_Event>(this, nullptr, SourceGraph);
+			UK2Node_Event* DelegateEvent = CompilerContext.SpawnIntermediateNode<UK2Node_Event>(this, SourceGraph);
 			DelegateEvent->EventReference.SetFromField<UFunction>(TargetFunction, false);
 			DelegateEvent->CustomFunctionName = GetDelegateTargetEntryPointName();
 			DelegateEvent->bInternalEvent = true;
diff --git a/Engine/Source/Editor/BlueprintGraph/Private/K2Node_InputAction.cpp b/Engine/Source/Editor/BlueprintGraph/Private/K2Node_InputAction.cpp
--- a/Engine/Source/Editor/BlueprintGraph/Private/K2Node_InputAction.cpp
+++ b/Engine/Source/Editor/BlueprintGraph/Private/K2Node_InputAction.cpp
@@ -188,7 +188,7 @@ void UK2Node_InputAction::ExpandNode(FKismetCompilerContext& CompilerContext, UE
 		{			
 			UEdGraphPin *EachPin = (*PinIt).Pin;
 			// Create the input touch event
-			UK2Node_InputActionEvent* InputActionEvent = CompilerContext.SpawnIntermediateEventNode<UK2Node_InputActionEvent>(this, EachPin, SourceGraph);
+			UK2Node_InputActionEvent* InputActionEvent = CompilerContext.SpawnIntermediateNode<UK2Node_InputActionEvent>(this, SourceGraph);
 			InputActionEvent->CustomFunctionName = FName( *FString::Printf(TEXT("InpActEvt_%s_%s"), *InputActionName.ToString(), *InputActionEvent->GetName()));
 			InputActionEvent->InputActionName = InputActionName;
 			InputActionEvent->bConsumeInput = bConsumeInput;
@@ -221,7 +221,7 @@ void UK2Node_InputAction::ExpandNode(FKismetCompilerContext& CompilerContext, UE
 	
 		if (InputActionPin->LinkedTo.Num() > 0)
 		{
-			UK2Node_InputActionEvent* InputActionEvent = CompilerContext.SpawnIntermediateEventNode<UK2Node_InputActionEvent>(this, InputActionPin, SourceGraph);
+			UK2Node_InputActionEvent* InputActionEvent = CompilerContext.SpawnIntermediateNode<UK2Node_InputActionEvent>(this, SourceGraph);
 			InputActionEvent->CustomFunctionName = FName( *FString::Printf(TEXT("InpActEvt_%s_%s"), *InputActionName.ToString(), *InputActionEvent->GetName()));
 			InputActionEvent->InputActionName = InputActionName;
 			InputActionEvent->bConsumeInput = bConsumeInput;
diff --git a/Engine/Source/Editor/BlueprintGraph/Private/K2Node_InputKey.cpp b/Engine/Source/Editor/BlueprintGraph/Private/K2Node_InputKey.cpp
--- a/Engine/Source/Editor/BlueprintGraph/Private/K2Node_InputKey.cpp
+++ b/Engine/Source/Editor/BlueprintGraph/Private/K2Node_InputKey.cpp
@@ -290,7 +290,7 @@ void UK2Node_InputKey::ExpandNode(FKismetCompilerContext& CompilerContext, UEdGr
 		{
 			UEdGraphPin *EachPin = (*PinIt).Pin;
 			// Create the input touch event
-			UK2Node_InputKeyEvent* InputKeyEvent = CompilerContext.SpawnIntermediateEventNode<UK2Node_InputKeyEvent>(this, EachPin, SourceGraph);
+			UK2Node_InputKeyEvent* InputKeyEvent = CompilerContext.SpawnIntermediateNode<UK2Node_InputKeyEvent>(this, SourceGraph);
 			const FName ModifierName = GetModifierName();
 			if ( ModifierName != NAME_None )
 			{
@@ -335,7 +335,7 @@ void UK2Node_InputKey::ExpandNode(FKismetCompilerContext& CompilerContext, UEdGr
 
 		if (InputKeyPin->LinkedTo.Num() > 0)
 		{
-			UK2Node_InputKeyEvent* InputKeyEvent = CompilerContext.SpawnIntermediateEventNode<UK2Node_InputKeyEvent>(this, InputKeyPin, SourceGraph);
+			UK2Node_InputKeyEvent* InputKeyEvent = CompilerContext.SpawnIntermediateNode<UK2Node_InputKeyEvent>(this, SourceGraph);
 			const FName ModifierName = GetModifierName();
 			if ( ModifierName != NAME_None )
 			{
diff --git a/Engine/Source/Editor/BlueprintGraph/Private/K2Node_InputTouch.cpp b/Engine/Source/Editor/BlueprintGraph/Private/K2Node_InputTouch.cpp
--- a/Engine/Source/Editor/BlueprintGraph/Private/K2Node_InputTouch.cpp
+++ b/Engine/Source/Editor/BlueprintGraph/Private/K2Node_InputTouch.cpp
@@ -232,7 +232,7 @@ void UK2Node_InputTouch::ExpandNode(FKismetCompilerContext& CompilerContext, UEd
 		{			
 			UEdGraphPin *EachPin= (*PinIt).Pin;
 			// Create the input touch event
-			UK2Node_InputTouchEvent* InputTouchEvent = CompilerContext.SpawnIntermediateEventNode<UK2Node_InputTouchEvent>(this, EachPin, SourceGraph);
+			UK2Node_InputTouchEvent* InputTouchEvent = CompilerContext.SpawnIntermediateNode<UK2Node_InputTouchEvent>(this, SourceGraph);
 			InputTouchEvent->CustomFunctionName = FName(*FString::Printf(TEXT("InpTchEvt_%s"), *EachPin->GetName()));
 			InputTouchEvent->bConsumeInput = bConsumeInput;
 			InputTouchEvent->bExecuteWhenPaused = bExecuteWhenPaused;
@@ -273,7 +273,7 @@ void UK2Node_InputTouch::ExpandNode(FKismetCompilerContext& CompilerContext, UEd
 		
 		if (InputTouchPin->LinkedTo.Num() > 0)
 		{
-			UK2Node_InputTouchEvent* InputTouchEvent = CompilerContext.SpawnIntermediateEventNode<UK2Node_InputTouchEvent>(this, InputTouchPin, SourceGraph);
+			UK2Node_InputTouchEvent* InputTouchEvent = CompilerContext.SpawnIntermediateNode<UK2Node_InputTouchEvent>(this, SourceGraph);
 			InputTouchEvent->CustomFunctionName = FName( *FString::Printf(TEXT("InpTchEvt_%s"), *InputTouchEvent->GetName()));
 			InputTouchEvent->InputKeyEvent = InputEvent;
 			InputTouchEvent->bConsumeInput = bConsumeInput;
diff --git a/Engine/Source/Editor/BlueprintGraph/Private/K2Node_LoadAsset.cpp b/Engine/Source/Editor/BlueprintGraph/Private/K2Node_LoadAsset.cpp
--- a/Engine/Source/Editor/BlueprintGraph/Private/K2Node_LoadAsset.cpp
+++ b/Engine/Source/Editor/BlueprintGraph/Private/K2Node_LoadAsset.cpp
@@ -155,7 +155,7 @@ void UK2Node_LoadAsset::ExpandNode(class FKismetCompilerContext& CompilerContext
 
 	// Create OnLoadEvent
 	const FName DelegateOnLoadedParamName(TEXT("OnLoaded"));
-	UK2Node_CustomEvent* OnLoadEventNode = CompilerContext.SpawnIntermediateEventNode<UK2Node_CustomEvent>(this, CallFunctionAssetPin, SourceGraph);
+	UK2Node_CustomEvent* OnLoadEventNode = CompilerContext.SpawnIntermediateNode<UK2Node_CustomEvent>(this, SourceGraph);
 	OnLoadEventNode->CustomFunctionName = *FString::Printf(TEXT("OnLoaded_%s"), *CompilerContext.GetGuid(this));
 	OnLoadEventNode->AllocateDefaultPins();
 	{
diff --git a/Engine/Source/Editor/KismetCompiler/Private/KismetCompiler.cpp b/Engine/Source/Editor/KismetCompiler/Private/KismetCompiler.cpp
--- a/Engine/Source/Editor/KismetCompiler/Private/KismetCompiler.cpp
+++ b/Engine/Source/Editor/KismetCompiler/Private/KismetCompiler.cpp
@@ -3251,12 +3251,7 @@ void FKismetCompilerContext::BuildDynamicBindingObjects(UBlueprintGeneratedClass
  */ 
 void FKismetCompilerContext::CreatePinEventNodeForTimelineFunction(UK2Node_Timeline* TimelineNode, UEdGraph* SourceGraph, FName FunctionName, const FName PinName, FName ExecFuncName)
 {
-	UEdGraphPin* SourcePin = nullptr;
-	if (UK2Node_Timeline* SourceNode = Cast<UK2Node_Timeline>(MessageLog.FindSourceObject(TimelineNode)))
-	{
-		SourcePin = SourceNode->FindPin(PinName);
-	}
-	UK2Node_Event* TimelineEventNode = SpawnIntermediateEventNode<UK2Node_Event>(TimelineNode, SourcePin, SourceGraph);
+	UK2Node_Event* TimelineEventNode = SpawnIntermediateNode<UK2Node_Event>(TimelineNode, SourceGraph);
 	TimelineEventNode->EventReference.SetExternalMember(FunctionName, UTimelineComponent::StaticClass());
 	TimelineEventNode->CustomFunctionName = FunctionName; // Make sure we name this function the thing we are expecting
 	TimelineEventNode->bInternalEvent = true;
@@ -4033,7 +4028,7 @@ void FKismetCompilerContext::CreateAndProcessUbergraph()
 				if (!bFoundEntry)
 				{
 					// Create an entry node stub, so that we have a entry point for interfaces to call to
-					UK2Node_Event* EventNode = SpawnIntermediateEventNode<UK2Node_Event>(nullptr, nullptr, ConsolidatedEventGraph);
+					UK2Node_Event* EventNode = SpawnIntermediateNode<UK2Node_Event>(nullptr, ConsolidatedEventGraph);
 					EventNode->EventReference.SetExternalMember(FunctionName, InterfaceDesc.Interface);
 					EventNode->bOverrideFunction = true;
 					EventNode->AllocateDefaultPins();
diff --git a/Engine/Source/Editor/KismetCompiler/Public/KismetCompiler.h b/Engine/Source/Editor/KismetCompiler/Public/KismetCompiler.h
--- a/Engine/Source/Editor/KismetCompiler/Public/KismetCompiler.h
+++ b/Engine/Source/Editor/KismetCompiler/Public/KismetCompiler.h
@@ -147,10 +147,6 @@ public:
 	TMap<UEdGraphNode*, UEdGraphNode*> CallsIntoUbergraph;
 	int32 bIsFullCompile:1;
 
-	// Source pin to latent/delegate entry point
-	TMap<UEdGraphNode*, UK2Node_Event*> SourceNodeToExpansionEvent;
-	TMap<UEdGraphPin*, UK2Node_Event*> SourcePinToExpansionEvent;
-
 	// Map that can be used to find the macro node that spawned a provided node, 
 	// if any. Macro instances can have more macros inside of them, so entries 
 	// in this map may chain (i.e. values may also need to be used as keys to find
@@ -214,9 +210,9 @@ public:
 
 	// Spawns an intermediate node associated with the source node (for error purposes)
 	template <typename NodeType>
-	NodeType* SpawnIntermediateNode(UEdGraphNode* SourceNode, UEdGraph* ParentGraph = NULL)
+	NodeType* SpawnIntermediateNode(UEdGraphNode* SourceNode, UEdGraph* ParentGraph = nullptr)
 	{
-		if (ParentGraph == NULL)
+		if (ParentGraph == nullptr)
 		{
 			ParentGraph = SourceNode->GetGraph();
 		}
@@ -231,37 +227,11 @@ public:
 		return Result;
 	}
 
-	// Spawns an intermediate event node associated with the source node (for error purposes)
 	template <typename NodeType>
+	UE_DEPRECATED(5.4, "SpawnIntermediateEventNode is equivalent to SpawnIntermediateNode, this redundant function has been deprecated.")
 	NodeType* SpawnIntermediateEventNode(UEdGraphNode* SourceNode, UEdGraphPin* SourcePin = nullptr, UEdGraph* ParentGraph = nullptr)
 	{
-		if (ParentGraph == nullptr && SourceNode != nullptr)
-		{
-			ParentGraph = SourceNode->GetGraph();
-		}
-
-		NodeType* Result = ParentGraph->CreateIntermediateNode<NodeType>();
-		//check (Cast<UK2Node_Event>(Result) != nullptr); -- Removed to avoid any fallout, will replace with care later
-		MessageLog.NotifyIntermediateObjectCreation(Result, SourceNode); // this might be useful to track back function entry nodes to events.
-		Result->CreateNewGuid();
-		// Track all expansion step generated events.
-		if (bIsFullCompile)
-		{
-			if (SourcePin)
-			{
-				UEdGraphNode* TrueSourceNode = Cast<UEdGraphNode>(MessageLog.FindSourceObject(SourcePin->GetOwningNode()));
-				UEdGraphPin* TrueSourcePin = TrueSourceNode->FindPin(SourcePin->GetFName());
-				SourcePinToExpansionEvent.Add(TrueSourcePin) = Cast<UK2Node_Event>(Result);
-			}
-			else if (SourceNode)
-			{
-				UEdGraphNode* TrueSourceNode = Cast<UEdGraphNode>(MessageLog.FindSourceObject(SourceNode));
-				SourceNodeToExpansionEvent.Add(TrueSourceNode) = Cast<UK2Node_Event>(Result);
-			}
-		}
-		AutoAssignNodePosition(Result);
-
-		return Result;
+		return SpawnIntermediateNode<NodeType>(SourceNode, ParentGraph);
 	}
 
 	/**
