MVVM: Fix bad copy paste when building the setter for FieldNotify.

diff --git a/Engine/Source/Editor/KismetCompiler/Private/KismetCompilerMisc.cpp b/Engine/Source/Editor/KismetCompiler/Private/KismetCompilerMisc.cpp
--- a/Engine/Source/Editor/KismetCompiler/Private/KismetCompilerMisc.cpp
+++ b/Engine/Source/Editor/KismetCompiler/Private/KismetCompilerMisc.cpp
@@ -1069,7 +1069,7 @@ TTuple<UEdGraphPin*, UEdGraphPin*> FKismetCompilerUtilities::GenerateFieldNotifi
 		CompilerContext.CopyPinLinksToIntermediate(*SourceNode->FindPinChecked(VariableReference.GetMemberName(), EGPD_Input), *NewValuePin);
 
 		bool bUseReferenceByRef = NewValuePin->LinkedTo.Num() != 0;
-		K2Schema->TrySetDefaultValue(*CallFuncNode->FindPinChecked(FName("NewValueByRef")), bHasLocalRepNotify ? TEXT("True") : TEXT("False"));
+		K2Schema->TrySetDefaultValue(*CallFuncNode->FindPinChecked(FName("NewValueByRef")), bUseReferenceByRef ? TEXT("True") : TEXT("False"));
 	}
 
 	// Set Net args pin
