Allow mapping a non-native class to a specific kind of Blueprint

#rnx

diff --git a/Engine/Source/Editor/KismetCompiler/Private/KismetCompilerModule.cpp b/Engine/Source/Editor/KismetCompiler/Private/KismetCompilerModule.cpp
--- a/Engine/Source/Editor/KismetCompiler/Private/KismetCompilerModule.cpp
+++ b/Engine/Source/Editor/KismetCompiler/Private/KismetCompilerModule.cpp
@@ -54,13 +54,13 @@ public:
 	virtual void GetSubclassesWithDifferingBlueprintTypes(UClass* Class, TSet<const UClass*>& OutMismatchedSubclasses) const override;
 	// End implementation
 
-	static TSubclassOf<UBlueprint> FindBlueprintType(UClass* ForClass, const TMap<UClass*, TSubclassOf<UBlueprint>>& FromMap);
+	static TSubclassOf<UBlueprint> FindBlueprintType(UClass* ForClass, const TMap<FTopLevelAssetPath, TSubclassOf<UBlueprint>>& FromMap);
 private:
 	// these are all pointers to native reflection data, so don't require gc visibility
 	// this will frustrate hotreload, though - hot reload of objects used as keys or values
 	// doesn't really work anyway:
-	TMap<UClass*, TSubclassOf<UBlueprint>> ClassToBPType;
-	TMap<UClass*, TSubclassOf<UBlueprint>> ClassToEditorBPType;
+	TMap<FTopLevelAssetPath, TSubclassOf<UBlueprint>> ClassToBPType;
+	TMap<FTopLevelAssetPath, TSubclassOf<UBlueprint>> ClassToEditorBPType;
 	TMap<TSubclassOf<UBlueprint>, TSubclassOf<UBlueprintGeneratedClass>> BPTypeToBPGCType;
 
 	TArray<IBlueprintCompiler*> Compilers;
@@ -214,7 +214,7 @@ void FKismet2CompilerModule::OverrideBPTypeForClass(UClass* Class, TSubclassOf<U
 {
 	check(Class && BlueprintType);
 	#if DO_CHECK
-	if (const TSubclassOf<UBlueprint>* ExistingBlueprintType = ClassToBPType.Find(Class))
+	if (const TSubclassOf<UBlueprint>* ExistingBlueprintType = ClassToBPType.Find(Class->GetClassPathName()))
 	{
 		ensureMsgf(false,
 			TEXT("Ambiguous mapping attempting to add %s to %s when mapping to %s exists"),
@@ -224,14 +224,14 @@ void FKismet2CompilerModule::OverrideBPTypeForClass(UClass* Class, TSubclassOf<U
 	}
 	#endif // DO_CHECK
 
-	ClassToBPType.Add(Class, BlueprintType);
+	ClassToBPType.Add(Class->GetClassPathName(), BlueprintType);
 }
 
 void FKismet2CompilerModule::OverrideBPTypeForClassInEditor(UClass* Class, TSubclassOf<UBlueprint> BlueprintType)
 {
 	check(Class && BlueprintType);
 #if DO_CHECK
-	if (const TSubclassOf<UBlueprint>* ExistingBlueprintType = ClassToEditorBPType.Find(Class))
+	if (const TSubclassOf<UBlueprint>* ExistingBlueprintType = ClassToEditorBPType.Find(Class->GetClassPathName()))
 	{
 		ensureMsgf(false,
 			TEXT("Ambiguous mapping attempting to add %s to %s when mapping to %s exists"),
@@ -241,7 +241,7 @@ void FKismet2CompilerModule::OverrideBPTypeForClassInEditor(UClass* Class, TSubc
 	}
 #endif // DO_CHECK
 
-	ClassToEditorBPType.Add(Class, BlueprintType);
+	ClassToEditorBPType.Add(Class->GetClassPathName(), BlueprintType);
 }
 
 void FKismet2CompilerModule::OverrideBPGCTypeForBPType(TSubclassOf<UBlueprint> BlueprintType, TSubclassOf<UBlueprintGeneratedClass> BPGCType)
@@ -392,11 +392,12 @@ void FKismet2CompilerModule::GetSubclassesWithDifferingBlueprintTypes(UClass* Cl
 	UClass* BPGeneratedClass;
 	GetBlueprintTypesForClass(Class, BPClass, BPGeneratedClass);
 	
-	auto CheckClassToBPTypeMap = [](const TMap<UClass*, TSubclassOf<UBlueprint>>& Map, const UClass* Class, const UClass* BPClass, TSet<const UClass*>& Result)
+	auto CheckClassToBPTypeMap = [](const TMap<FTopLevelAssetPath, TSubclassOf<UBlueprint>>& Map, const UClass* Class, const UClass* BPClass, TSet<const UClass*>& Result)
 	{
-		for (const auto&[SupportedClass, OtherBPClass] : Map)
+		for (const TTuple<FTopLevelAssetPath, TSubclassOf<UBlueprint>>& Pair : Map)
 		{
-			if (OtherBPClass != BPClass && SupportedClass->IsChildOf(Class))
+			if (const UClass* SupportedClass = FindObject<UClass>(Pair.Key);
+				Pair.Value != BPClass && SupportedClass && SupportedClass->IsChildOf(Class))
 			{
 				Result.Add(SupportedClass);
 			}
@@ -406,12 +407,12 @@ void FKismet2CompilerModule::GetSubclassesWithDifferingBlueprintTypes(UClass* Cl
 	CheckClassToBPTypeMap(ClassToEditorBPType, Class, BPClass, OutMismatchedSubclasses);
 }
 
-TSubclassOf<UBlueprint> FKismet2CompilerModule::FindBlueprintType(UClass* ForClass, const TMap<UClass*, TSubclassOf<UBlueprint>>& FromMap)
+TSubclassOf<UBlueprint> FKismet2CompilerModule::FindBlueprintType(UClass* ForClass, const TMap<FTopLevelAssetPath, TSubclassOf<UBlueprint>>& FromMap)
 {
 	UClass* Iter = ForClass;
 	while (Iter)
 	{
-		if (const TSubclassOf<UBlueprint>* BPType = FromMap.Find(Iter))
+		if (const TSubclassOf<UBlueprint>* BPType = FromMap.Find(Iter->GetClassPathName()))
 		{
 			return *BPType;
 		}
