#UE Improve logging about misconfigured ConvertibleDelegates

diff --git a/Engine/Source/Editor/KismetCompiler/Private/KismetCompiler.cpp b/Engine/Source/Editor/KismetCompiler/Private/KismetCompiler.cpp
--- a/Engine/Source/Editor/KismetCompiler/Private/KismetCompiler.cpp
+++ b/Engine/Source/Editor/KismetCompiler/Private/KismetCompiler.cpp
@@ -4814,7 +4814,14 @@ void FKismetCompilerContext::CompileClassLayout(EInternalCompilerFlags InternalF
 	CreateFunctionList();
 
 	// Function list creation should process captured variables. Something went wrong if we missed any.
-	UE_CLOG(!ConvertibleDelegates.IsEmpty(), LogK2Compiler, Warning, TEXT("%d convertible delegates were not processed during class layout compilation."), ConvertibleDelegates.Num());
+	if (!ConvertibleDelegates.IsEmpty())
+	{
+		UE_LOG(LogK2Compiler, Warning, TEXT("%d convertible delegates were not processed during class layout compilation. Listing delegates in log below."), ConvertibleDelegates.Num());
+		for (auto DelegateIt = ConvertibleDelegates.CreateConstIterator(); DelegateIt; ++DelegateIt)
+		{
+			UE_LOG(LogK2Compiler, Display, TEXT("  Node:%s Function:%s Variable:%s"), *GetPathNameSafe(DelegateIt.Key()), *DelegateIt.Value().ProxyFunctionName.ToString(), *DelegateIt.Value().CapturedVariableName.ToString());
+		}
+	}
 
 	// Precompile the functions
 	// Handle delegates signatures first, because they are needed by other functions
