Improve missing-property error messages to use full paths and use specified scope instead of self-class.

diff --git a/Engine/Source/Editor/BlueprintGraph/Private/K2Node_Variable.cpp b/Engine/Source/Editor/BlueprintGraph/Private/K2Node_Variable.cpp
--- a/Engine/Source/Editor/BlueprintGraph/Private/K2Node_Variable.cpp
+++ b/Engine/Source/Editor/BlueprintGraph/Private/K2Node_Variable.cpp
@@ -561,10 +561,10 @@ void UK2Node_Variable::ValidateNodeDuringCompilation(class FCompilerResultsLog&
 			UBlueprint* Blueprint = GetBlueprint();
 			if (Blueprint != nullptr)
 			{
-				OwnerName = Blueprint->GetName();
+				OwnerName = Blueprint->GetPathName();
 				if (UClass* VarOwnerClass = VariableReference.GetMemberParentClass(Blueprint->GeneratedClass))
 				{
-					OwnerName = VarOwnerClass->GetName();
+					OwnerName = VarOwnerClass->GetPathName();
 				}
 			}
 
diff --git a/Engine/Source/Editor/KismetCompiler/Private/KismetCompilerMisc.cpp b/Engine/Source/Editor/KismetCompiler/Private/KismetCompilerMisc.cpp
--- a/Engine/Source/Editor/KismetCompiler/Private/KismetCompilerMisc.cpp
+++ b/Engine/Source/Editor/KismetCompiler/Private/KismetCompilerMisc.cpp
@@ -557,7 +557,8 @@ FProperty* FKismetCompilerUtilities::FindPropertyInScope(UStruct* Scope, UEdGrap
 	}
 	else if (!FKismetCompilerUtilities::IsMissingMemberPotentiallyLoading(Cast<UBlueprint>(SelfClass->ClassGeneratedBy), Scope))
 	{
-		MessageLog.Error(*FText::Format(LOCTEXT("PropertyNotFound_Error", "The property associated with @@ could not be found in '{0}'"), FText::FromString(SelfClass->GetPathName())).ToString(), Pin);
+		UObject* MessageScope = Scope ? Scope : SelfClass;
+		MessageLog.Error(*FText::Format(LOCTEXT("PropertyNotFound_Error", "The property associated with @@ could not be found in '{0}'"), FText::FromString(MessageScope->GetPathName())).ToString(), Pin);
 	}
 
 	return nullptr;
