Remove redundant call to PropagateParentBlueprintDefaults when using the compilation manager, this saves us .5s of boot time in large projects and reduces compilation side effects

diff --git a/Engine/Source/Editor/KismetCompiler/Private/KismetCompiler.cpp b/Engine/Source/Editor/KismetCompiler/Private/KismetCompiler.cpp
--- a/Engine/Source/Editor/KismetCompiler/Private/KismetCompiler.cpp
+++ b/Engine/Source/Editor/KismetCompiler/Private/KismetCompiler.cpp
@@ -4861,7 +4861,7 @@ void FKismetCompilerContext::CompileFunctions(EInternalCompilerFlags InternalFla
 	// This is phase two, so we want to generated locals if PostponeLocalsGenerationUntilPhaseTwo is set:
 	const bool bGenerateLocals = !!(InternalFlags & EInternalCompilerFlags::PostponeLocalsGenerationUntilPhaseTwo);
 	// Don't propagate values to CDO if we're going to do that in reinstancing:
-	bool bPropagateValuesToCDO = !(InternalFlags & EInternalCompilerFlags::PostponeDefaultObjectAssignmentUntilReinstancing);
+	const bool bPropagateValuesToCDO = !(InternalFlags & EInternalCompilerFlags::PostponeDefaultObjectAssignmentUntilReinstancing);
 	// Don't RefreshExternalBlueprintDependencyNodes if the calling code has done so already:
 	const bool bSkipRefreshExternalBlueprintDependencyNodes = !!(InternalFlags & EInternalCompilerFlags::SkipRefreshExternalBlueprintDependencyNodes);
 	FKismetCompilerVMBackend Backend_VM(Blueprint, Schema, *this);
@@ -4985,10 +4985,10 @@ void FKismetCompilerContext::CompileFunctions(EInternalCompilerFlags InternalFla
 		UObject* NewCDO = NewClass->GetDefaultObject();
 
 		// Copy over the CDO properties if we're not already regenerating on load.  In that case, the copy will be done after compile on load is complete
-		FBlueprintEditorUtils::PropagateParentBlueprintDefaults(NewClass);
-
 		if(bPropagateValuesToCDO)
 		{
+			FBlueprintEditorUtils::PropagateParentBlueprintDefaults(NewClass);
+
 			if( !Blueprint->HasAnyFlags(RF_BeingRegenerated) )
 			{
 				// Propagate the old CDO's properties to the new
