PR #10869: Fixed BP compilation with deprecated event dispatchers (multicast delegates).

Regression introduced at CL 24063670.

diff --git a/Engine/Source/Editor/KismetCompiler/Private/KismetCompiler.cpp b/Engine/Source/Editor/KismetCompiler/Private/KismetCompiler.cpp
--- a/Engine/Source/Editor/KismetCompiler/Private/KismetCompiler.cpp
+++ b/Engine/Source/Editor/KismetCompiler/Private/KismetCompiler.cpp
@@ -2533,9 +2533,10 @@ void FKismetCompilerContext::PrecompileFunction(FKismetFunctionContext& Context,
 			Context.Function->FunctionFlags |= FUNC_Delegate;
 
 			// We really don't want to find our parent's delegate property and accidentally 
-			// overwrite the signature function, so provide EFieldIterationFlags::None:
+			// overwrite the signature function, so EFieldIterationFlags::IncludeSuper is excluded:
+			const EFieldIterationFlags IterationFlags = EFieldIterationFlags::Default & ~EFieldIterationFlags::IncludeSuper;
 			if (FMulticastDelegateProperty* Property = FindFProperty<FMulticastDelegateProperty>(
-					NewClass, Context.DelegateSignatureName, EFieldIterationFlags::None))
+					NewClass, Context.DelegateSignatureName, IterationFlags))
 			{
 				Property->SignatureFunction = Context.Function;
 			}
