UE-194830 - The Blueprint compiler is not consistent with UHT when generating DefaultToInstanced variables

Object properties that reference a DefaultToInstanced class also need to set the "EditInline" metadata to true. Otherwise, this property can't be instanced in the editor. This mimics UHT's behavior for UhtObjectProperty.

diff --git a/Engine/Source/Editor/KismetCompiler/Private/KismetCompilerMisc.cpp b/Engine/Source/Editor/KismetCompiler/Private/KismetCompilerMisc.cpp
--- a/Engine/Source/Editor/KismetCompiler/Private/KismetCompilerMisc.cpp
+++ b/Engine/Source/Editor/KismetCompiler/Private/KismetCompilerMisc.cpp
@@ -1269,6 +1269,7 @@ FProperty* FKismetCompilerUtilities::CreatePrimitiveProperty(FFieldVariant Prope
 				if (SubType->HasAnyClassFlags(CLASS_DefaultToInstanced))
 				{
 					NewPropertyObj->SetPropertyFlags(CPF_InstancedReference);
+					NewPropertyObj->SetMetaData(TEXT("EditInline"), TEXT("true"));
 				}
 
 				// we want to use this setter function instead of setting the 
