Follow-up fix for CL 30398348: don't set EditInline on actor components.

CL 30398348 introduced a loophole for actor component instancing. By setting a component variable to "EditInline", it was now possible to instance a component in both the Details panel *and* the Components panel. Additionally, instancing in the Details panel would actually outer the component to the CDO instead of the UClass.

Note that component variables will default to using SPropertyEditorAsset in the Details panel, which is also misleading, since it implies that a non-null default value can be set on the component property. UE-203267 was created to improve the messaging here.

diff --git a/Engine/Source/Editor/KismetCompiler/Private/KismetCompilerMisc.cpp b/Engine/Source/Editor/KismetCompiler/Private/KismetCompilerMisc.cpp
--- a/Engine/Source/Editor/KismetCompiler/Private/KismetCompilerMisc.cpp
+++ b/Engine/Source/Editor/KismetCompiler/Private/KismetCompilerMisc.cpp
@@ -44,7 +44,6 @@
 #include "KismetCastingUtils.h"
 #include "KismetCompiledFunctionContext.h"
 #include "KismetCompiler.h"
-
 #include "K2Node_EnumLiteral.h"
 #include "Kismet/KismetArrayLibrary.h"
 #include "Kismet2/KismetReinstanceUtilities.h"
@@ -52,6 +51,7 @@
 #include "Kismet2/StructureEditorUtils.h"
 #include "ObjectTools.h"
 #include "BlueprintEditorSettings.h"
+#include "Components/ActorComponent.h"
 
 #define LOCTEXT_NAMESPACE "KismetCompiler"
 
@@ -1269,7 +1269,16 @@ FProperty* FKismetCompilerUtilities::CreatePrimitiveProperty(FFieldVariant Prope
 				if (SubType->HasAnyClassFlags(CLASS_DefaultToInstanced))
 				{
 					NewPropertyObj->SetPropertyFlags(CPF_InstancedReference);
-					NewPropertyObj->SetMetaData(TEXT("EditInline"), TEXT("true"));
+
+					// Actor components should only be instanced by the SCS editor.
+					// 
+					// Default actor components are outered to the generated BP class instead of the CDO.
+					// If we set "EditInline" on actor components, we would actually outer them to the CDO.
+					// This would lead to various serialization and instancing issues.
+					if (!SubType->IsChildOf<UActorComponent>())
+					{
+						NewPropertyObj->SetMetaData(TEXT("EditInline"), TEXT("true"));
+					}
 				}
 
 				// we want to use this setter function instead of setting the 
