UE-203280 - Blueprints from UE4 with dangling delegate reroute node crashes on UE5 load

GetLinkedDelegateSignature has a fairly aggressive assert. It's entirely possible that a knot (reroute) node doesn't have an output linked. In those cases, we can safely break out of the loop, since there won't be a linked delegate signature.

diff --git a/Engine/Source/Editor/KismetCompiler/Private/KismetCompiler.cpp b/Engine/Source/Editor/KismetCompiler/Private/KismetCompiler.cpp
--- a/Engine/Source/Editor/KismetCompiler/Private/KismetCompiler.cpp
+++ b/Engine/Source/Editor/KismetCompiler/Private/KismetCompiler.cpp
@@ -322,9 +322,10 @@ namespace UE::KismetCompiler::Private
 				check(KnotNode->GetOutputPin());
 				UEdGraphPin* OutputPin = KnotNode->GetOutputPin();
 				check(OutputPin);
-				check(OutputPin->LinkedTo.Num() > 0);
-				
-				LinkedNode = OutputPin->LinkedTo[0]->GetOwningNode();
+
+				// Knots don't necessarily have to be linked.
+				// In these cases, there simply isn't a corresponding signature to reference.
+				LinkedNode = (OutputPin->LinkedTo.Num() > 0) ? OutputPin->LinkedTo[0]->GetOwningNode() : nullptr;
 			}
 
 			// If the delegate pin has been linked, then we should run into either a
