Deprecated FKismetCompilerUtilities::ValidateProperEndExecutionPath.

This function was once used by the BP compiler for validating that return nodes are linked, but was removed several years ago.

We'll likely want to resurrect the warning at some point, but we probably won't use this implementation.

diff --git a/Engine/Source/Editor/KismetCompiler/Private/KismetCompilerMisc.cpp b/Engine/Source/Editor/KismetCompiler/Private/KismetCompilerMisc.cpp
--- a/Engine/Source/Editor/KismetCompiler/Private/KismetCompilerMisc.cpp
+++ b/Engine/Source/Editor/KismetCompiler/Private/KismetCompilerMisc.cpp
@@ -1682,204 +1682,9 @@ FFieldVariant FKismetCompilerUtilities::CheckPropertyNameOnScope(UStruct* Scope,
 	return FFieldVariant();
 }
 
-/** Checks if the execution path ends with a Return node */
 void FKismetCompilerUtilities::ValidateProperEndExecutionPath(FKismetFunctionContext& Context)
 {
-	struct FRecrursiveHelper
-	{
-		static bool IsExecutionSequence(const UEdGraphNode* Node)
-		{
-			return Node && (UK2Node_ExecutionSequence::StaticClass() == Node->GetClass()); // no "SourceNode->IsA<UK2Node_ExecutionSequence>()" because MultiGate is based on ExecutionSequence
-		}
-
-		static void CheckPathEnding(const UK2Node* StartingNode, TSet<const UK2Node*>& VisitedNodes, FKismetFunctionContext& InContext, bool bPathShouldEndWithReturn, TSet<const UK2Node*>& BreakableNodesSeeds)
-		{
-			const UK2Node* CurrentNode = StartingNode;
-			while (CurrentNode)
-			{
-				const UK2Node* SourceNode = CurrentNode;
-				CurrentNode = nullptr;
-
-				bool bAlreadyVisited = false;
-				VisitedNodes.Add(SourceNode, &bAlreadyVisited);
-				if (!bAlreadyVisited && !SourceNode->IsA<UK2Node_FunctionResult>())
-				{
-					const bool bIsExecutionSequence = IsExecutionSequence(SourceNode); 
-					for (UEdGraphPin* CurrentPin : SourceNode->Pins)
-					{
-						if (CurrentPin
-							&& (CurrentPin->Direction == EEdGraphPinDirection::EGPD_Output)
-							&& (CurrentPin->PinType.PinCategory == UEdGraphSchema_K2::PC_Exec))
-						{
-							if (!CurrentPin->LinkedTo.Num())
-							{
-								if (!bIsExecutionSequence)
-								{
-									BreakableNodesSeeds.Add(SourceNode);
-								}
-								if (bPathShouldEndWithReturn && !bIsExecutionSequence)
-								{
-									InContext.MessageLog.Note(*LOCTEXT("ExecutionEnd_Note", "The execution path doesn't end with a return node. @@").ToString(), CurrentPin);
-								}
-								continue;
-							}
-							UEdGraphPin* LinkedPin = CurrentPin->LinkedTo[0];
-							const UK2Node* NextNode = ensure(LinkedPin) ? Cast<const UK2Node>(LinkedPin->GetOwningNodeUnchecked()) : nullptr;
-							ensure(NextNode);
-							if (CurrentNode)
-							{
-								FRecrursiveHelper::CheckPathEnding(CurrentNode, VisitedNodes, InContext, bPathShouldEndWithReturn && !bIsExecutionSequence, BreakableNodesSeeds);
-							}
-							CurrentNode = NextNode;
-						}
-					}
-				}
-			}
-		}
-
-		static void GatherBreakableNodes(const UK2Node* StartingNode, TSet<const UK2Node*>& BreakableNodes, FKismetFunctionContext& InContext)
-		{
-			const UK2Node* CurrentNode = StartingNode;
-			while (CurrentNode)
-			{
-				const UK2Node* SourceNode = CurrentNode;
-				CurrentNode = nullptr;
-
-				bool bAlreadyVisited = false;
-				BreakableNodes.Add(SourceNode, &bAlreadyVisited);
-				if (!bAlreadyVisited)
-				{
-					for (UEdGraphPin* CurrentPin : SourceNode->Pins)
-					{
-						if (CurrentPin
-							&& (CurrentPin->Direction == EEdGraphPinDirection::EGPD_Input)
-							&& (CurrentPin->PinType.PinCategory == UEdGraphSchema_K2::PC_Exec)
-							&& CurrentPin->LinkedTo.Num())
-						{
-							for (UEdGraphPin* LinkedPin : CurrentPin->LinkedTo)
-							{
-								const UK2Node* NextNode = ensure(LinkedPin) ? Cast<const UK2Node>(LinkedPin->GetOwningNodeUnchecked()) : nullptr;
-								ensure(NextNode);
-								if (!FRecrursiveHelper::IsExecutionSequence(NextNode))
-								{
-									if (CurrentNode)
-									{
-										GatherBreakableNodes(CurrentNode, BreakableNodes, InContext);
-									}
-									CurrentNode = NextNode;
-								}
-							}
-						}
-					}
-				}
-			}
-		}
-
-		static void GatherBreakableNodesSeedsFromSequences(TSet<const UK2Node_ExecutionSequence*>& UnBreakableExecutionSequenceNodes
-			, TSet<const UK2Node*>& BreakableNodesSeeds
-			, TSet<const UK2Node*>& BreakableNodes
-			, FKismetFunctionContext& InContext)
-		{
-			for (const UK2Node_ExecutionSequence* SequenceNode : UnBreakableExecutionSequenceNodes)
-			{
-				bool bIsBreakable = true;
-				// Sequence is breakable when all it's outputs are breakable
-				for (UEdGraphPin* CurrentPin : SequenceNode->Pins)
-				{
-					if (CurrentPin
-						&& (CurrentPin->Direction == EEdGraphPinDirection::EGPD_Output)
-						&& (CurrentPin->PinType.PinCategory == UEdGraphSchema_K2::PC_Exec)
-						&& CurrentPin->LinkedTo.Num())
-					{
-						UEdGraphPin* LinkedPin = CurrentPin->LinkedTo[0];
-						const UK2Node* NextNode = ensure(LinkedPin) ? Cast<const UK2Node>(LinkedPin->GetOwningNodeUnchecked()) : nullptr;
-						ensure(NextNode);
-						if (!BreakableNodes.Contains(NextNode))
-						{
-							bIsBreakable = false;
-							break;
-						}
-					}
-				}
-
-				if (bIsBreakable)
-				{
-					bool bWasAlreadyBreakable = false;
-					BreakableNodesSeeds.Add(SequenceNode, &bWasAlreadyBreakable);
-					ensure(!bWasAlreadyBreakable);
-					int32 WasRemoved = UnBreakableExecutionSequenceNodes.Remove(SequenceNode);
-					ensure(WasRemoved);
-				}
-			}
-		}
-
-		static void CheckDeadExecutionPath(TSet<const UK2Node*>& BreakableNodesSeeds, FKismetFunctionContext& InContext)
-		{
-			TSet<const UK2Node_ExecutionSequence*> UnBreakableExecutionSequenceNodes;
-			for (UEdGraphNode* Node : InContext.SourceGraph->Nodes)
-			{
-				if (FRecrursiveHelper::IsExecutionSequence(Node))
-				{
-					UnBreakableExecutionSequenceNodes.Add(Cast<UK2Node_ExecutionSequence>(Node));
-				}
-			}
-
-			TSet<const UK2Node*> BreakableNodes;
-			while (BreakableNodesSeeds.Num())
-			{
-				for (const UK2Node* StartingNode : BreakableNodesSeeds)
-				{
-					GatherBreakableNodes(StartingNode, BreakableNodes, InContext);
-				}
-				BreakableNodesSeeds.Empty();
-				FRecrursiveHelper::GatherBreakableNodesSeedsFromSequences(UnBreakableExecutionSequenceNodes, BreakableNodesSeeds, BreakableNodes, InContext);
-			}
-
-			for (const UK2Node_ExecutionSequence* UnBreakableExecutionSequenceNode : UnBreakableExecutionSequenceNodes)
-			{
-				bool bUnBreakableOutputWasFound = false;
-				for (UEdGraphPin* CurrentPin : UnBreakableExecutionSequenceNode->Pins)
-				{
-					if (CurrentPin
-						&& (CurrentPin->Direction == EEdGraphPinDirection::EGPD_Output)
-						&& (CurrentPin->PinType.PinCategory == UEdGraphSchema_K2::PC_Exec)
-						&& CurrentPin->LinkedTo.Num())
-					{
-						if (bUnBreakableOutputWasFound)
-						{
-							InContext.MessageLog.Note(*LOCTEXT("DeadExecution_Note", "The path is never executed. @@").ToString(), CurrentPin);
-							break;
-						}
-
-						UEdGraphPin* LinkedPin = CurrentPin->LinkedTo[0];
-						const UK2Node* NextNode = ensure(LinkedPin) ? Cast<const UK2Node>(LinkedPin->GetOwningNodeUnchecked()) : nullptr;
-						ensure(NextNode);
-						if (!BreakableNodes.Contains(NextNode))
-						{
-							bUnBreakableOutputWasFound = true;
-						}
-					}
-				}
-			}
-		}
-	};
-
-	// Function is designed for multiple return nodes.
-	if (!Context.IsEventGraph() && Context.SourceGraph && Context.Schema)
-	{
-		TArray<UK2Node_FunctionResult*> ReturnNodes;
-		Context.SourceGraph->GetNodesOfClass(ReturnNodes);
-		if (ReturnNodes.Num() && ensure(Context.EntryPoint))
-		{
-			TSet<const UK2Node*> VisitedNodes, BreakableNodesSeeds;
-			FRecrursiveHelper::CheckPathEnding(Context.EntryPoint, VisitedNodes, Context, true, BreakableNodesSeeds);
-
-			// A non-pure node, that lies on a execution path, that may result with "EndThread" state, is called Breakable.
-			// The execution path between the node and Return node can be broken.
-
-			FRecrursiveHelper::CheckDeadExecutionPath(BreakableNodesSeeds, Context);
-		}
-	}
+	ensureMsgf(false, TEXT("ValidateProperEndExecutionPath has been deprecated"));
 }
 
 void FKismetCompilerUtilities::DetectValuesReturnedByRef(const UFunction* Func, const UK2Node * Node, FCompilerResultsLog& MessageLog)
diff --git a/Engine/Source/Editor/KismetCompiler/Public/KismetCompilerMisc.h b/Engine/Source/Editor/KismetCompiler/Public/KismetCompilerMisc.h
--- a/Engine/Source/Editor/KismetCompiler/Public/KismetCompilerMisc.h
+++ b/Engine/Source/Editor/KismetCompiler/Public/KismetCompilerMisc.h
@@ -122,7 +122,7 @@ public:
 	/** Create Kismet assignment statement with proper object <-> interface cast */
 	static void CreateObjectAssignmentStatement(FKismetFunctionContext& Context, UEdGraphNode* Node, FBPTerminal* SrcTerm, FBPTerminal* DstTerm, UEdGraphPin* DstPin = nullptr);
 
-	/** Checks if each execution path ends with a Return node */
+	UE_DEPRECATED(5.4, "ValidateProperEndExecutionPath is deprecated.")
 	static void ValidateProperEndExecutionPath(FKismetFunctionContext& Context);
 
 	/** Generate an error for non-const output parameters */
