Removed unused code

#rnx

diff --git a/Engine/Source/Editor/KismetCompiler/Private/KismetCompilerModule.cpp b/Engine/Source/Editor/KismetCompiler/Private/KismetCompilerModule.cpp
--- a/Engine/Source/Editor/KismetCompiler/Private/KismetCompilerModule.cpp
+++ b/Engine/Source/Editor/KismetCompiler/Private/KismetCompilerModule.cpp
@@ -65,24 +65,6 @@ private:
 
 IMPLEMENT_MODULE( FKismet2CompilerModule, KismetCompiler );
 
-struct FBlueprintIsBeingCompiledHelper
-{
-private:
-	UBlueprint* Blueprint;
-public:
-	FBlueprintIsBeingCompiledHelper(UBlueprint* InBlueprint) : Blueprint(InBlueprint)
-	{
-		check(NULL != Blueprint);
-		check(!Blueprint->bBeingCompiled);
-		Blueprint->bBeingCompiled = true;
-	}
-
-	~FBlueprintIsBeingCompiledHelper()
-	{
-		Blueprint->bBeingCompiled = false;
-	}
-};
-
 // Compiles a blueprint.
 
 void FKismet2CompilerModule::CompileStructure(UUserDefinedStruct* Struct, FCompilerResultsLog& Results)
