[Blueprint]
 - Remove loading flags from old classes and sub objects after they are trashed

diff --git a/Engine/Source/Editor/KismetCompiler/Private/KismetCompiler.cpp b/Engine/Source/Editor/KismetCompiler/Private/KismetCompiler.cpp
--- a/Engine/Source/Editor/KismetCompiler/Private/KismetCompiler.cpp
+++ b/Engine/Source/Editor/KismetCompiler/Private/KismetCompiler.cpp
@@ -660,6 +660,7 @@ void FKismetCompilerContext::CleanAndSanitizeClass(UBlueprintGeneratedClass* Cla
 		FName TransientCDOName = MakeUniqueObjectName(GetTransientPackage(), TransientClass, FName(*TransientCDOString));
 		InOldCDO->Rename(*TransientCDOName.ToString(), GetTransientPackage(), RenFlags);
 		FLinkerLoad::InvalidateExport(InOldCDO);
+		InOldCDO->ClearFlags(RF_NeedLoad | RF_NeedPostLoad | RF_NeedPostLoadSubobjects);
 	}
 
 	// Purge all subobjects (properties, functions, params) of the class, as they will be regenerated
@@ -723,6 +724,7 @@ void FKismetCompilerContext::CleanAndSanitizeClass(UBlueprintGeneratedClass* Cla
 		FName NewSubobjectName = MakeUniqueObjectName(TransientClass, CurrSubObj->GetClass(), CurrSubObj->GetFName());
 		CurrSubObj->Rename(*NewSubobjectName.ToString(), TransientClass, RenFlags | REN_ForceNoResetLoaders);
 		FLinkerLoad::InvalidateExport(CurrSubObj);
+		CurrSubObj->ClearFlags(RF_NeedLoad | RF_NeedPostLoad | RF_NeedPostLoadSubobjects);
 	}
 
 	// Purge the class to get it back to a "base" state
