Add an argument to ignore  blueprint warnings (will be simple log messages).

- Geared towards engine testing against older assets.

diff --git a/Engine/Source/Editor/KismetCompiler/Private/EdGraphCompiler.cpp b/Engine/Source/Editor/KismetCompiler/Private/EdGraphCompiler.cpp
--- a/Engine/Source/Editor/KismetCompiler/Private/EdGraphCompiler.cpp
+++ b/Engine/Source/Editor/KismetCompiler/Private/EdGraphCompiler.cpp
@@ -21,6 +21,18 @@
 #include "Kismet2/CompilerResultsLog.h"
 #include "Templates/SubclassOf.h"
 #include "UObject/ObjectPtr.h"
+#include "Misc/CommandLine.h"
+
+namespace UE::Private::GraphCompilerContext
+{
+	// Builds against old assets routinely generate deprecation messages, fixing them would defeat the point of running compat-testing builds.
+	// Since we expect deprecated BP to continue to be functional, provide a way to make the warnings simple log messages.
+	bool IgnoreBlueprintDeprecationWarnings()
+	{
+		const static bool bIgnoreBpWarnings = FParse::Param(FCommandLine::Get(), TEXT("IgnoreBlueprintDeprecationWarnings"));
+		return bIgnoreBpWarnings;
+	}
+}
 
 //////////////////////////////////////////////////////////////////////////
 // FGraphCompilerContext
@@ -84,13 +96,16 @@ void FGraphCompilerContext::ValidateNode(const UEdGraphNode* Node) const
 		{
 			switch (Response.MessageType)
 			{
+			case EEdGraphNodeDeprecationMessageType::Warning:
+				if (LIKELY(!UE::Private::GraphCompilerContext::IgnoreBlueprintDeprecationWarnings()))
+				{
+					MessageLog.Warning(*Response.MessageText.ToString(), Node);
+					break;
+				}
+				// else intentional fall-through to Note
 			case EEdGraphNodeDeprecationMessageType::Note:
 				MessageLog.Note(*Response.MessageText.ToString(), Node);
 				break;
-
-			case EEdGraphNodeDeprecationMessageType::Warning:
-				MessageLog.Warning(*Response.MessageText.ToString(), Node);
-				break;
 			}
 		}
 	}
