Honor UE_WITH_OBJECT_HANDLE_LATE_RESOLVE when working with FObjectPtr internals in bpvmbackend
#rnx

diff --git a/Engine/Source/Editor/KismetCompiler/Private/KismetCompilerVMBackend.cpp b/Engine/Source/Editor/KismetCompiler/Private/KismetCompilerVMBackend.cpp
--- a/Engine/Source/Editor/KismetCompiler/Private/KismetCompilerVMBackend.cpp
+++ b/Engine/Source/Editor/KismetCompiler/Private/KismetCompilerVMBackend.cpp
@@ -94,7 +94,11 @@ public:
 
 	FArchive& operator<<(FObjectPtr& Res) override
 	{
+#if UE_WITH_OBJECT_HANDLE_LATE_RESOLVE
 		ScriptPointerType D = (ScriptPointerType)Res.GetHandle().PointerOrRef;
+#else
+		ScriptPointerType D = (ScriptPointerType)Res.GetHandle();
+#endif
 		FArchive& Ar = *this;
 
 		Ar << D;
