[Blueprint]: Fix spawn object nodes with linked class inputs to always generate assignments for unlinked exposed parameter inputs even if the pin's value matches the underlying (static) type's default value, since the actual class may not be resolved until runtime.

diff --git a/Engine/Source/Editor/BlueprintGraph/Private/K2Node_AddComponentByClass.cpp b/Engine/Source/Editor/BlueprintGraph/Private/K2Node_AddComponentByClass.cpp
--- a/Engine/Source/Editor/BlueprintGraph/Private/K2Node_AddComponentByClass.cpp
+++ b/Engine/Source/Editor/BlueprintGraph/Private/K2Node_AddComponentByClass.cpp
@@ -226,7 +226,7 @@ void UK2Node_AddComponentByClass::ExpandNode(class FKismetCompilerContext& Compi
 	//////////////////////////////////////////////////////////////////////////
 	// create 'set var' nodes
 
-	UEdGraphPin* LastThen = FKismetCompilerUtilities::GenerateAssignmentNodes(CompilerContext, SourceGraph, CallAddComponentByClassNode, this, CallAddComponentByClassResult, ClassToSpawn);
+	UEdGraphPin* LastThen = FKismetCompilerUtilities::GenerateAssignmentNodes(CompilerContext, SourceGraph, CallAddComponentByClassNode, this, CallAddComponentByClassResult, ClassToSpawn, CallAddComponentByClassTypePin);
 
 	if (LastThen != CallAddComponentByClassNode->GetThenPin())
 	{
diff --git a/Engine/Source/Editor/BlueprintGraph/Private/K2Node_GenericCreateObject.cpp b/Engine/Source/Editor/BlueprintGraph/Private/K2Node_GenericCreateObject.cpp
--- a/Engine/Source/Editor/BlueprintGraph/Private/K2Node_GenericCreateObject.cpp
+++ b/Engine/Source/Editor/BlueprintGraph/Private/K2Node_GenericCreateObject.cpp
@@ -107,10 +107,11 @@ void UK2Node_GenericCreateObject::ExpandNode(class FKismetCompilerContext& Compi
 		bSucceeded &= SpawnExecPin && CallExecPin && CompilerContext.MovePinLinksToIntermediate(*SpawnExecPin, *CallExecPin).CanSafeConnect();
 	}
 
+	UEdGraphPin* CallClassPin = nullptr;
 	//connect class
 	{
 		UEdGraphPin* SpawnClassPin = GetClassPin();
-		UEdGraphPin* CallClassPin = CallCreateNode->FindPin(TEXT("ObjectClass"));
+		CallClassPin = CallCreateNode->FindPin(TEXT("ObjectClass"));
 		bSucceeded &= SpawnClassPin && CallClassPin && CompilerContext.MovePinLinksToIntermediate(*SpawnClassPin, *CallClassPin).CanSafeConnect();
 	}
 		
@@ -137,7 +138,7 @@ void UK2Node_GenericCreateObject::ExpandNode(class FKismetCompilerContext& Compi
 
 	//assign exposed values and connect then
 	{
-		UEdGraphPin* LastThen = FKismetCompilerUtilities::GenerateAssignmentNodes(CompilerContext, SourceGraph, CallCreateNode, this, CallResultPin, ClassToSpawn);
+		UEdGraphPin* LastThen = FKismetCompilerUtilities::GenerateAssignmentNodes(CompilerContext, SourceGraph, CallCreateNode, this, CallResultPin, ClassToSpawn, CallClassPin);
 		UEdGraphPin* SpawnNodeThen = GetThenPin();
 		bSucceeded &= SpawnNodeThen && LastThen && CompilerContext.MovePinLinksToIntermediate(*SpawnNodeThen, *LastThen).CanSafeConnect();
 	}
diff --git a/Engine/Source/Editor/BlueprintGraph/Private/K2Node_SpawnActorFromClass.cpp b/Engine/Source/Editor/BlueprintGraph/Private/K2Node_SpawnActorFromClass.cpp
--- a/Engine/Source/Editor/BlueprintGraph/Private/K2Node_SpawnActorFromClass.cpp
+++ b/Engine/Source/Editor/BlueprintGraph/Private/K2Node_SpawnActorFromClass.cpp
@@ -572,7 +572,7 @@ void UK2Node_SpawnActorFromClass::ExpandNode(class FKismetCompilerContext& Compi
 	// create 'set var' nodes
 
 	// Get 'result' pin from 'begin spawn', this is the actual actor we want to set properties on
-	UEdGraphPin* LastThen = FKismetCompilerUtilities::GenerateAssignmentNodes(CompilerContext, SourceGraph, CallBeginSpawnNode, SpawnNode, CallBeginResult, ClassToSpawn );
+	UEdGraphPin* LastThen = FKismetCompilerUtilities::GenerateAssignmentNodes(CompilerContext, SourceGraph, CallBeginSpawnNode, SpawnNode, CallBeginResult, ClassToSpawn, CallBeginActorClassPin);
 
 	// Make exec connection between 'then' on last node and 'finish'
 	LastThen->MakeLinkTo(CallFinishExec);
diff --git a/Engine/Source/Editor/KismetCompiler/Private/KismetCompilerMisc.cpp b/Engine/Source/Editor/KismetCompiler/Private/KismetCompilerMisc.cpp
--- a/Engine/Source/Editor/KismetCompiler/Private/KismetCompilerMisc.cpp
+++ b/Engine/Source/Editor/KismetCompiler/Private/KismetCompilerMisc.cpp
@@ -763,7 +763,7 @@ bool FKismetCompilerUtilities::ValidateSelfCompatibility(const UEdGraphPin* Pin,
 	return true;
 }
 
-UEdGraphPin* FKismetCompilerUtilities::GenerateAssignmentNodes(class FKismetCompilerContext& CompilerContext, UEdGraph* SourceGraph, UK2Node_CallFunction* CallBeginSpawnNode, UEdGraphNode* SpawnNode, UEdGraphPin* CallBeginResult, const UClass* ForClass )
+UEdGraphPin* FKismetCompilerUtilities::GenerateAssignmentNodes(class FKismetCompilerContext& CompilerContext, UEdGraph* SourceGraph, UK2Node_CallFunction* CallBeginSpawnNode, UEdGraphNode* SpawnNode, UEdGraphPin* CallBeginResult, const UClass* ForClass, const UEdGraphPin* CallBeginClassInput)
 {
 	static const FName ObjectParamName(TEXT("Object"));
 	static const FName ValueParamName(TEXT("Value"));
@@ -772,13 +772,14 @@ UEdGraphPin* FKismetCompilerUtilities::GenerateAssignmentNodes(class FKismetComp
 	const UEdGraphSchema_K2* Schema = CompilerContext.GetSchema();
 	UEdGraphPin* LastThen = CallBeginSpawnNode->GetThenPin();
 
+	const bool bIsClassInputPinLinked = CallBeginClassInput && CallBeginClassInput->LinkedTo.Num() > 0;
+
 	// Create 'set var by name' nodes and hook them up
 	for (int32 PinIdx = 0; PinIdx < SpawnNode->Pins.Num(); PinIdx++)
 	{
 		// Only create 'set param by name' node if this pin is linked to something
 		UEdGraphPin* OrgPin = SpawnNode->Pins[PinIdx];
-		if (NULL == CallBeginSpawnNode->FindPin(OrgPin->PinName) &&
-			(OrgPin->LinkedTo.Num() > 0 || !OrgPin->DoesDefaultValueMatchAutogenerated()))
+		if (!CallBeginSpawnNode->FindPin(OrgPin->PinName))
 		{
 			FProperty* Property = FindFProperty<FProperty>(ForClass, OrgPin->PinName);
 			// NULL property indicates that this pin was part of the original node, not the 
@@ -788,7 +789,11 @@ UEdGraphPin* FKismetCompilerUtilities::GenerateAssignmentNodes(class FKismetComp
 				continue;
 			}
 
-			if( OrgPin->LinkedTo.Num() == 0 )
+			// If the class input pin is linked, then 'ForClass' represents a base type, but the actual type might be a derived class
+			// that won't get resolved until runtime. In that case, we can't use the base type's CDO to avoid generating assignment
+			// statements for unlinked parameter inputs that match the base type's default value for those parameters, since the
+			// derived type might store a different default value, which would otherwise be used if we don't explicitly assign it.
+			if (!bIsClassInputPinLinked && OrgPin->LinkedTo.Num() == 0)
 			{
 				// We don't want to generate an assignment node unless the default value 
 				// differs from the value in the CDO:
diff --git a/Engine/Source/Editor/KismetCompiler/Public/KismetCompilerMisc.h b/Engine/Source/Editor/KismetCompiler/Public/KismetCompilerMisc.h
--- a/Engine/Source/Editor/KismetCompiler/Public/KismetCompilerMisc.h
+++ b/Engine/Source/Editor/KismetCompiler/Public/KismetCompilerMisc.h
@@ -111,7 +111,7 @@ public:
 	static bool ValidateSelfCompatibility(const UEdGraphPin* Pin, FKismetFunctionContext& Context);
 
 	/** Create 'set var by name' nodes and hook them up - used to set values when components are added or actor are created at run time. Returns the 'last then' pin of the assignment nodes */
-	static UEdGraphPin* GenerateAssignmentNodes( class FKismetCompilerContext& CompilerContext, UEdGraph* SourceGraph, UK2Node_CallFunction* CallBeginSpawnNode, UEdGraphNode* SpawnNode, UEdGraphPin* CallBeginResult, const UClass* ForClass );
+	static UEdGraphPin* GenerateAssignmentNodes(class FKismetCompilerContext& CompilerContext, UEdGraph* SourceGraph, UK2Node_CallFunction* CallBeginSpawnNode, UEdGraphNode* SpawnNode, UEdGraphPin* CallBeginResult, const UClass* ForClass, const UEdGraphPin* CallBeginClassInput = nullptr);
 
 	/** Create node that replace regular setter and use the SetPropertyValueAndBroadcast. */
 	static TTuple<UEdGraphPin*, UEdGraphPin*> GenerateFieldNotificationSetNode(FKismetCompilerContext& CompilerContext, UEdGraph* SourceGraph, UEdGraphNode* SourceNode, UEdGraphPin* SelfPin, FProperty* VariableProperty, const FMemberReference& VariableReference, bool bHasLocalRepNotify, bool bShouldFlushDormancyOnSet, bool bIsNetProperty);
diff --git a/Engine/Source/Editor/UMGEditor/Private/Nodes/K2Node_CreateDragDropOperation.cpp b/Engine/Source/Editor/UMGEditor/Private/Nodes/K2Node_CreateDragDropOperation.cpp
--- a/Engine/Source/Editor/UMGEditor/Private/Nodes/K2Node_CreateDragDropOperation.cpp
+++ b/Engine/Source/Editor/UMGEditor/Private/Nodes/K2Node_CreateDragDropOperation.cpp
@@ -121,7 +121,7 @@ void UK2Node_CreateDragDropOperation::ExpandNode(class FKismetCompilerContext& C
 	// create 'set var' nodes
 
 	// Get 'result' pin from 'begin spawn', this is the actual actor we want to set properties on
-	UEdGraphPin* LastThen = FKismetCompilerUtilities::GenerateAssignmentNodes(CompilerContext, SourceGraph, CallCreateNode, CreateOpNode, CallCreateResult, GetClassToSpawn());
+	UEdGraphPin* LastThen = FKismetCompilerUtilities::GenerateAssignmentNodes(CompilerContext, SourceGraph, CallCreateNode, CreateOpNode, CallCreateResult, GetClassToSpawn(), CallCreateOperationClassPin);
 
 	// Move 'then' connection from create widget node to the last 'then'
 	CompilerContext.MovePinLinksToIntermediate(*SpawnNodeThen, *LastThen);
diff --git a/Engine/Source/Editor/UMGEditor/Private/Nodes/K2Node_CreateWidget.cpp b/Engine/Source/Editor/UMGEditor/Private/Nodes/K2Node_CreateWidget.cpp
--- a/Engine/Source/Editor/UMGEditor/Private/Nodes/K2Node_CreateWidget.cpp
+++ b/Engine/Source/Editor/UMGEditor/Private/Nodes/K2Node_CreateWidget.cpp
@@ -150,7 +150,7 @@ void UK2Node_CreateWidget::ExpandNode(class FKismetCompilerContext& CompilerCont
 	// create 'set var' nodes
 
 	// Get 'result' pin from 'begin spawn', this is the actual actor we want to set properties on
-	UEdGraphPin* LastThen = FKismetCompilerUtilities::GenerateAssignmentNodes(CompilerContext, SourceGraph, CallCreateNode, CreateWidgetNode, CallCreateResult, ClassToSpawn);
+	UEdGraphPin* LastThen = FKismetCompilerUtilities::GenerateAssignmentNodes(CompilerContext, SourceGraph, CallCreateNode, CreateWidgetNode, CallCreateResult, ClassToSpawn, CallCreateWidgetTypePin);
 
 	// Move 'then' connection from create widget node to the last 'then'
 	CompilerContext.MovePinLinksToIntermediate(*SpawnNodeThen, *LastThen);
