Avoided a TArray copy when removing duplicate functions in FKismetCompilerVMBackend::GenerateCodeFromClass.

diff --git a/Engine/Source/Editor/KismetCompiler/Private/KismetCompilerVMBackend.cpp b/Engine/Source/Editor/KismetCompiler/Private/KismetCompilerVMBackend.cpp
--- a/Engine/Source/Editor/KismetCompiler/Private/KismetCompilerVMBackend.cpp
+++ b/Engine/Source/Editor/KismetCompiler/Private/KismetCompilerVMBackend.cpp
@@ -2298,7 +2298,7 @@ void FKismetCompilerVMBackend::GenerateCodeFromClass(UClass* SourceClass, TIndir
 
 	// Remove duplicates from CalledFunctions:
 	UBlueprintGeneratedClass* ClassBeingBuilt = CastChecked<UBlueprintGeneratedClass>(SourceClass);
-	TSet<UFunction*> Unique(ClassBeingBuilt->CalledFunctions);
+	TSet<TObjectPtr<UFunction>> Unique(ClassBeingBuilt->CalledFunctions);
 	ClassBeingBuilt->CalledFunctions = Unique.Array();
 }
 
