Unshelved from pending changelist '35264463':

Prune event graph functions with no entry points

[RN] Empty event graphs are now more aggressively pruned, can be disabled via [Blueprints] bDisableLateEventGraphPruning

diff --git a/Engine/Source/Editor/KismetCompiler/Private/KismetCompiler.cpp b/Engine/Source/Editor/KismetCompiler/Private/KismetCompiler.cpp
--- a/Engine/Source/Editor/KismetCompiler/Private/KismetCompiler.cpp
+++ b/Engine/Source/Editor/KismetCompiler/Private/KismetCompiler.cpp
@@ -6,6 +6,7 @@
 
 
 #include "KismetCompiler.h"
+#include "Algo/Find.h"
 #include "Engine/BlueprintGeneratedClass.h"
 #include "Misc/CoreMisc.h"
 #include "Components/ActorComponent.h"
@@ -3978,6 +3979,25 @@ void FKismetCompilerContext::VerifyValidOverrideFunction(const UEdGraph* Graph)
 	}
 }
 
+namespace UE::Private
+{
+
+static bool CanEventGraphBePruned(const UEdGraph* ConsolidatedEventGraph)
+{
+	static const FBoolConfigValueHelper DisableLateEventGraphPruning(TEXT("Blueprints"), TEXT("bDisableLateEventGraphPruning"), GEngineIni);
+	if(DisableLateEventGraphPruning)
+	{
+		return false;
+	}
+
+	return Algo::FindByPredicate(ConsolidatedEventGraph->Nodes, 
+		[](const UEdGraphNode* Node)
+		{
+			return Cast<UK2Node_Event>(Node) != nullptr;
+		}) == nullptr;
+}
+
+}
 
 // Merges pages and creates function stubs, etc... from the ubergraph entry points
 void FKismetCompilerContext::CreateAndProcessUbergraph()
@@ -4075,6 +4095,13 @@ void FKismetCompilerContext::CreateAndProcessUbergraph()
 
 		// If a function in the graph cannot be overridden/placed as event make sure that it is not.
 		VerifyValidOverrideEvent(ConsolidatedEventGraph);
+		
+		// expansion may have removed all of the event entry points, in which case nothing will
+		// be placed in the ubergraph function and we can prune it:
+		if(UE::Private::CanEventGraphBePruned(ConsolidatedEventGraph))
+		{
+			return;
+		}
 
 		// Do some cursory validation (pin types match, inputs to outputs, pins never point to their parent node, etc...)
 		{
