Fix lifetime of newly added compiler callback

#rnx

diff --git a/Engine/Source/Editor/KismetCompiler/Private/KismetCompiler.cpp b/Engine/Source/Editor/KismetCompiler/Private/KismetCompiler.cpp
--- a/Engine/Source/Editor/KismetCompiler/Private/KismetCompiler.cpp
+++ b/Engine/Source/Editor/KismetCompiler/Private/KismetCompiler.cpp
@@ -5438,7 +5438,7 @@ void FKismetCompilerContext::PostCDOCompiled(const UObject::FPostCDOCompiledCont
 		NewClass->ClassDefaultObject->PostCDOCompiled();
 		PRAGMA_ENABLE_DEPRECATION_WARNINGS;
 		
-		for (const TFunctionRef<void(const UObject::FPostCDOCompiledContext&, UObject*)>& Step : PostCDOCompileSteps)
+		for (const TFunction<void(const UObject::FPostCDOCompiledContext&, UObject*)>& Step : PostCDOCompileSteps)
 		{
 			Step(Context, NewClass->ClassDefaultObject);
 		}
@@ -5450,9 +5450,9 @@ void FKismetCompilerContext::PostCDOCompiled(const UObject::FPostCDOCompiledCont
 	OnPostCDOCompiled(Context);
 }
 
-void FKismetCompilerContext::AddPostCDOCompiledStep(TFunctionRef<void(const UObject::FPostCDOCompiledContext&, UObject*)> StepFunction)
+void FKismetCompilerContext::AddPostCDOCompiledStep(TFunction<void(const UObject::FPostCDOCompiledContext&, UObject*)>&& StepFunction)
 {
-	PostCDOCompileSteps.Add(StepFunction);
+	PostCDOCompileSteps.Emplace(StepFunction);
 }
 
 void FKismetCompilerContext::Compile()
diff --git a/Engine/Source/Editor/KismetCompiler/Public/KismetCompiler.h b/Engine/Source/Editor/KismetCompiler/Public/KismetCompiler.h
--- a/Engine/Source/Editor/KismetCompiler/Public/KismetCompiler.h
+++ b/Engine/Source/Editor/KismetCompiler/Public/KismetCompiler.h
@@ -89,7 +89,7 @@ protected:
 	TMap< TSubclassOf<class UEdGraphNode>, FNodeHandlingFunctor*> NodeHandlers;
 
 	// Array of function refs to run after the CDO has been compiled in FKismetCompilerContext::PostCDOCompiled	
-	TArray<TFunctionRef<void(const UObject::FPostCDOCompiledContext&, UObject*)>> PostCDOCompileSteps;
+	TArray<TFunction<void(const UObject::FPostCDOCompiledContext&, UObject*)>> PostCDOCompileSteps;
 
 	// Map of properties created for timelines; to aid in debug data generation
 	TMap<class UTimelineTemplate*, class FProperty*> TimelineToMemberVariableMap;
@@ -211,7 +211,7 @@ public:
 	/**
 	* Adds the given "StepFunction" to be run after the CDO has been compiled in FKismetCompilerContext::PostCDOCompiled.
 	*/
-	void AddPostCDOCompiledStep(TFunctionRef<void (const UObject::FPostCDOCompiledContext&, UObject*)> StepFunction);
+	void AddPostCDOCompiledStep(TFunction<void (const UObject::FPostCDOCompiledContext&, UObject*)>&& StepFunction);
 
 	/** Compile a blueprint into a class and a set of functions */
 	void Compile();
