PR #12255: Make FKismetCompilerUtilities::GenerateAssignmentNodes callable from other sources

The logic of GenerateAssignmentNodes does not require a Call Function node, only that the input has a "then" pin which is not required by all Call Function nodes. Making this less restrictive makes it easier to write new custom K2 nodes based on or similar in functionality to CreateObjectFromClass without requiring a Call Function node is the source of the object.

diff --git a/Engine/Source/Editor/KismetCompiler/Private/KismetCompilerMisc.cpp b/Engine/Source/Editor/KismetCompiler/Private/KismetCompilerMisc.cpp
--- a/Engine/Source/Editor/KismetCompiler/Private/KismetCompilerMisc.cpp
+++ b/Engine/Source/Editor/KismetCompiler/Private/KismetCompilerMisc.cpp
@@ -763,7 +763,7 @@ bool FKismetCompilerUtilities::ValidateSelfCompatibility(const UEdGraphPin* Pin,
 	return true;
 }
 
-UEdGraphPin* FKismetCompilerUtilities::GenerateAssignmentNodes(class FKismetCompilerContext& CompilerContext, UEdGraph* SourceGraph, UK2Node_CallFunction* CallBeginSpawnNode, UEdGraphNode* SpawnNode, UEdGraphPin* CallBeginResult, const UClass* ForClass, const UEdGraphPin* CallBeginClassInput)
+UEdGraphPin* FKismetCompilerUtilities::GenerateAssignmentNodes(class FKismetCompilerContext& CompilerContext, UEdGraph* SourceGraph, UK2Node* CallBeginSpawnNode, UEdGraphNode* SpawnNode, UEdGraphPin* CallBeginResult, const UClass* ForClass, const UEdGraphPin* CallBeginClassInput)
 {
 	static const FName ObjectParamName(TEXT("Object"));
 	static const FName ValueParamName(TEXT("Value"));
diff --git a/Engine/Source/Editor/KismetCompiler/Public/KismetCompilerMisc.h b/Engine/Source/Editor/KismetCompiler/Public/KismetCompilerMisc.h
--- a/Engine/Source/Editor/KismetCompiler/Public/KismetCompilerMisc.h
+++ b/Engine/Source/Editor/KismetCompiler/Public/KismetCompilerMisc.h
@@ -111,7 +111,7 @@ public:
 	static bool ValidateSelfCompatibility(const UEdGraphPin* Pin, FKismetFunctionContext& Context);
 
 	/** Create 'set var by name' nodes and hook them up - used to set values when components are added or actor are created at run time. Returns the 'last then' pin of the assignment nodes */
-	static UEdGraphPin* GenerateAssignmentNodes(class FKismetCompilerContext& CompilerContext, UEdGraph* SourceGraph, UK2Node_CallFunction* CallBeginSpawnNode, UEdGraphNode* SpawnNode, UEdGraphPin* CallBeginResult, const UClass* ForClass, const UEdGraphPin* CallBeginClassInput = nullptr);
+	static UEdGraphPin* GenerateAssignmentNodes(class FKismetCompilerContext& CompilerContext, UEdGraph* SourceGraph, UK2Node* CallBeginSpawnNode, UEdGraphNode* SpawnNode, UEdGraphPin* CallBeginResult, const UClass* ForClass, const UEdGraphPin* CallBeginClassInput = nullptr);
 
 	/** Create node that replace regular setter and use the SetPropertyValueAndBroadcast. */
 	static TTuple<UEdGraphPin*, UEdGraphPin*> GenerateFieldNotificationSetNode(FKismetCompilerContext& CompilerContext, UEdGraph* SourceGraph, UEdGraphNode* SourceNode, UEdGraphPin* SelfPin, FProperty* VariableProperty, const FMemberReference& VariableReference, bool bHasLocalRepNotify, bool bShouldFlushDormancyOnSet, bool bIsNetProperty);
