PR #12254: Adjust FKismetCompilerUtilities::GenerateAssignmentNodes to ignore output pins

GenerateAssignmentNodes gets confused by the existence of a pin who name matches a property name but cannot be connected to an assignment node input.

diff --git a/Engine/Source/Editor/KismetCompiler/Private/KismetCompilerMisc.cpp b/Engine/Source/Editor/KismetCompiler/Private/KismetCompilerMisc.cpp
--- a/Engine/Source/Editor/KismetCompiler/Private/KismetCompilerMisc.cpp
+++ b/Engine/Source/Editor/KismetCompiler/Private/KismetCompilerMisc.cpp
@@ -783,6 +783,12 @@ UEdGraphPin* FKismetCompilerUtilities::GenerateAssignmentNodes(class FKismetComp
 	{
 		// Only create 'set param by name' node if this pin is linked to something
 		UEdGraphPin* OrgPin = SpawnNode->Pins[PinIdx];
+
+		if (OrgPin->Direction == EGPD_Output)
+		{
+			continue;
+		}
+		
 		if (!CallBeginSpawnNode->FindPin(OrgPin->PinName))
 		{
 			FProperty* Property = FindFProperty<FProperty>(ForClass, OrgPin->PinName);
