Do not mark soft or weak object blueprint properties as EditInline/InstancedReference because this is not supported by the engine and will crash the editor.
This matches the behavior of UHT for native properties.

diff --git a/Engine/Source/Editor/KismetCompiler/Private/KismetCompilerMisc.cpp b/Engine/Source/Editor/KismetCompiler/Private/KismetCompilerMisc.cpp
--- a/Engine/Source/Editor/KismetCompiler/Private/KismetCompilerMisc.cpp
+++ b/Engine/Source/Editor/KismetCompiler/Private/KismetCompilerMisc.cpp
@@ -1282,21 +1282,21 @@ FProperty* FKismetCompilerUtilities::CreatePrimitiveProperty(FFieldVariant Prope
 					{
 						NewPropertyObj->SetPropertyFlags(CPF_TObjectPtrWrapper);
 					}
-				}
 
-				// Is the property a reference to something that should default to instanced?
-				if (SubType->HasAnyClassFlags(CLASS_DefaultToInstanced))
-				{
-					NewPropertyObj->SetPropertyFlags(CPF_InstancedReference);
-
-					// Actor components should only be instanced by the SCS editor.
-					// 
-					// Default actor components are outered to the generated BP class instead of the CDO.
-					// If we set "EditInline" on actor components, we would actually outer them to the CDO.
-					// This would lead to various serialization and instancing issues.
-					if (!SubType->IsChildOf<UActorComponent>())
+					// Is the property a reference to something that should default to instanced?
+					if (SubType->HasAnyClassFlags(CLASS_DefaultToInstanced))
 					{
-						NewPropertyObj->SetMetaData(TEXT("EditInline"), TEXT("true"));
+						NewPropertyObj->SetPropertyFlags(CPF_InstancedReference);
+
+						// Actor components should only be instanced by the SCS editor.
+						// 
+						// Default actor components are outered to the generated BP class instead of the CDO.
+						// If we set "EditInline" on actor components, we would actually outer them to the CDO.
+						// This would lead to various serialization and instancing issues.
+						if (!SubType->IsChildOf<UActorComponent>())
+						{
+							NewPropertyObj->SetMetaData(TEXT("EditInline"), TEXT("true"));
+						}
 					}
 				}
 
