Fix regression where objects/actors spawned in blueprint would incorrectly have object properties cleared because they were detected as different from CDO.
This fixes the automatic setting of actor Instigator from SpawnActor nodes as that was being overridden when the pin was not explicitly set to null

diff --git a/Engine/Source/Editor/KismetCompiler/Private/KismetCompilerMisc.cpp b/Engine/Source/Editor/KismetCompiler/Private/KismetCompilerMisc.cpp
--- a/Engine/Source/Editor/KismetCompiler/Private/KismetCompilerMisc.cpp
+++ b/Engine/Source/Editor/KismetCompiler/Private/KismetCompilerMisc.cpp
@@ -807,27 +807,38 @@ UEdGraphPin* FKismetCompilerUtilities::GenerateAssignmentNodes(class FKismetComp
 					// Some types require a connection for assignment (e.g. arrays).
 					continue;
 				}
-				else if (!bIsClassInputPinLinked)
+				// If the property is not editable in blueprint, always check the native CDO to handle cases like Instigator properly
+				else if (!bIsClassInputPinLinked || Property->HasAnyPropertyFlags(CPF_DisableEditOnTemplate) || !Property->HasAnyPropertyFlags(CPF_Edit))
 				{
 					// We don't want to generate an assignment node unless the default value 
 					// differs from the value in the CDO:
 					FString DefaultValueAsString;
 
-					if (FBlueprintCompilationManager::GetDefaultValue(ForClass, Property, DefaultValueAsString))
+					if (!FBlueprintCompilationManager::GetDefaultValue(ForClass, Property, DefaultValueAsString))
 					{
-						if (Schema->DoesDefaultValueMatch(*OrgPin, DefaultValueAsString))
+						if (ForClass->ClassDefaultObject)
 						{
-							continue;
+							FBlueprintEditorUtils::PropertyValueToString(Property, (uint8*)ForClass->ClassDefaultObject.Get(), DefaultValueAsString);
 						}
 					}
-					else if (ForClass->ClassDefaultObject)
+
+					// First check the string representation of the default value
+					if (Schema->DoesDefaultValueMatch(*OrgPin, DefaultValueAsString))
 					{
-						FBlueprintEditorUtils::PropertyValueToString(Property, (uint8*)ForClass->ClassDefaultObject.Get(), DefaultValueAsString);
+						continue;
+					}
 
-						if (DefaultValueAsString == OrgPin->GetDefaultAsString())
-						{
-							continue;
-						}
+					FString UseDefaultValue;
+					TObjectPtr<UObject> UseDefaultObject = nullptr;
+					FText UseDefaultText;
+					constexpr bool bPreserveTextIdentity = true;
+
+					// Next check if the converted default value would be the same to handle cases like None for object pointers
+					Schema->GetPinDefaultValuesFromString(OrgPin->PinType, OrgPin->GetOwningNodeUnchecked(), DefaultValueAsString, UseDefaultValue, UseDefaultObject, UseDefaultText, bPreserveTextIdentity);
+
+					if (OrgPin->DefaultValue.Equals(UseDefaultValue, ESearchCase::CaseSensitive) && OrgPin->DefaultObject == UseDefaultObject && OrgPin->DefaultTextValue.IdenticalTo(UseDefaultText))
+					{
+						continue;
 					}
 				}
 			}
