Fix widget blueprints never having the bAutomaticallyRegisterInputOnConstruction value set to true because it was being cleared in the Widget compiler.

This buggy behavior was introduced with 36334713, where we cleared it under the impression that if we did not, then if you deleted a Bluerpint event node from a widget then it would still have the flag set on it.

An additional call to UpdateCustomPropertyListForPostConstruction is needed in the kismet compiler here because if any of the user code in PostCDOCompiled makes a change to a native property's default value, the property link list needs to be updated so that instances of that object will get the correct value set in FObjectInitializer::InitProperties.

#rnx

diff --git a/Engine/Source/Editor/KismetCompiler/Private/KismetCompiler.cpp b/Engine/Source/Editor/KismetCompiler/Private/KismetCompiler.cpp
--- a/Engine/Source/Editor/KismetCompiler/Private/KismetCompiler.cpp
+++ b/Engine/Source/Editor/KismetCompiler/Private/KismetCompiler.cpp
@@ -5441,6 +5441,10 @@ void FKismetCompilerContext::PostCDOCompiled(const UObject::FPostCDOCompiledCont
 
 	// Allow children to customize PostCDOCompile:
 	OnPostCDOCompiled(Context);
+
+	// Run this function in case any native default property values have been changed on the CDO
+	// from all the user code above
+	NewClass->UpdateCustomPropertyListForPostConstruction();
 }
 
 void FKismetCompilerContext::AddPostCDOCompiledStep(TFunction<void(const UObject::FPostCDOCompiledContext&, UObject*)>&& StepFunction)
