Blueprint: Allow blueprint compilers to specify FBPVariableDescriptions to generate blueprint variables with during compilation

- Expose the guids of these generated variables to GetFieldNameFromClassByGuid/GetGuidFromClassByFieldName so that references to renamed generated variables can be fixed up

diff --git a/Engine/Source/Editor/Kismet/Private/BlueprintCompilationManager.cpp b/Engine/Source/Editor/Kismet/Private/BlueprintCompilationManager.cpp
--- a/Engine/Source/Editor/Kismet/Private/BlueprintCompilationManager.cpp
+++ b/Engine/Source/Editor/Kismet/Private/BlueprintCompilationManager.cpp
@@ -466,6 +466,7 @@ struct FCompilerData
 	bool IsSkeletonOnly() const { return JobType == ECompilationManagerJobType::SkeletonOnly; }
 	bool ShouldSetTemporaryBlueprintFlags() const { return JobType != ECompilationManagerJobType::RelinkOnly; }
 	bool ShouldResetErrorState() const { return JobType == ECompilationManagerJobType::Normal && InternalOptions.CompileType != EKismetCompileType::BytecodeOnly; }
+	bool ShouldResetAndPopulateGeneratedVariables() const { return JobType != ECompilationManagerJobType::RelinkOnly; }
 	bool ShouldValidate() const { return JobType == ECompilationManagerJobType::Normal; }
 	bool ShouldRegenerateSkeleton() const { return JobType != ECompilationManagerJobType::RelinkOnly; }
 	bool ShouldMarkUpToDateAfterSkeletonStage() const { return IsSkeletonOnly(); }
@@ -904,6 +905,11 @@ void FBlueprintCompilationManagerImpl::FlushCompilationQueueImpl(bool bSuppressB
 		// STAGE V: Validate
 		for (FCompilerData& CompilerData : CurrentlyCompilingBPs)
 		{
+			if (CompilerData.ShouldResetAndPopulateGeneratedVariables())
+			{
+				CompilerData.Compiler->ResetAndPopulateBlueprintGeneratedVariables();
+			}
+
 			if(!CompilerData.ShouldValidate())
 			{
 				continue;
diff --git a/Engine/Source/Editor/KismetCompiler/Private/KismetCompiler.cpp b/Engine/Source/Editor/KismetCompiler/Private/KismetCompiler.cpp
--- a/Engine/Source/Editor/KismetCompiler/Private/KismetCompiler.cpp
+++ b/Engine/Source/Editor/KismetCompiler/Private/KismetCompiler.cpp
@@ -1178,6 +1178,12 @@ void FKismetCompilerContext::ValidateClassPropertyDefaults()
 	}
 }
 
+void FKismetCompilerContext::ResetAndPopulateBlueprintGeneratedVariables()
+{
+	Blueprint->GeneratedVariables.Reset();
+	PopulateBlueprintGeneratedVariables();
+}
+
 void FKismetCompilerContext::ValidateTimelineNames()
 {
 	TSharedPtr<FKismetNameValidator> ParentBPNameValidator;
@@ -1232,9 +1238,14 @@ void FKismetCompilerContext::CreateClassVariablesFromBlueprint()
 	}
 	NewClass->FieldNotifies.Reset();
 
-	for (int32 i = 0; i < Blueprint->NewVariables.Num(); ++i)
+	// Create variables in reverse order starting with NewVariables
+	const int32 NumVariables = Blueprint->NewVariables.Num() + Blueprint->GeneratedVariables.Num();
+	for (int32 i = 0; i < NumVariables; ++i)
 	{
-		FBPVariableDescription& Variable = Blueprint->NewVariables[Blueprint->NewVariables.Num() - (i + 1)];
+		const bool bIsGeneratedVariable = i >= Blueprint->NewVariables.Num();
+		FBPVariableDescription& Variable = bIsGeneratedVariable
+			? Blueprint->GeneratedVariables[Blueprint->GeneratedVariables.Num() - (i + 1 - Blueprint->NewVariables.Num())]
+			: Blueprint->NewVariables[Blueprint->NewVariables.Num() - (i + 1)];
 
 		FProperty* NewProperty = CreateVariable(Variable.VarName, Variable.VarType);
 		if (NewProperty != NULL)
diff --git a/Engine/Source/Editor/KismetCompiler/Public/KismetCompiler.h b/Engine/Source/Editor/KismetCompiler/Public/KismetCompiler.h
--- a/Engine/Source/Editor/KismetCompiler/Public/KismetCompiler.h
+++ b/Engine/Source/Editor/KismetCompiler/Public/KismetCompiler.h
@@ -307,6 +307,9 @@ public:
 	*/
 	void ValidateClassPropertyDefaults();
 
+	/** Resets the blueprint's GeneratedVariables list and calls PopulateBlueprintGeneratedVariables */
+	void ResetAndPopulateBlueprintGeneratedVariables();
+
 	/** Creates a class variable for each entry in the Blueprint NewVars array */
 	virtual void CreateClassVariablesFromBlueprint();
 
@@ -342,6 +345,9 @@ protected:
 	 * @param OldCDO			Reference to the old CDO of the class, so we can copy the properties from it to the new class's CDO
 	 */
 	virtual void CleanAndSanitizeClass(UBlueprintGeneratedClass* ClassToClean, UObject*& OldCDO);
+	
+	/** Compilers are expected to populate the blueprint's full list of GeneratedVariables here as the list is reset at this point */
+	virtual void PopulateBlueprintGeneratedVariables() {}
 
 	struct KISMETCOMPILER_API FSubobjectCollection
 	{
diff --git a/Engine/Source/Runtime/Engine/Classes/Engine/Blueprint.h b/Engine/Source/Runtime/Engine/Classes/Engine/Blueprint.h
--- a/Engine/Source/Runtime/Engine/Classes/Engine/Blueprint.h
+++ b/Engine/Source/Runtime/Engine/Classes/Engine/Blueprint.h
@@ -600,7 +600,11 @@ PRAGMA_ENABLE_DEPRECATION_WARNINGS
 
 	/** Array of new variables to be added to generated class */
 	UPROPERTY()
-	TArray<struct FBPVariableDescription> NewVariables;
+	TArray<FBPVariableDescription> NewVariables;
+
+	/** Array of generated variables to be added to generated class */
+	UPROPERTY()
+	TArray<FBPVariableDescription> GeneratedVariables;
 
 	/** Array of user sorted categories */
 	UPROPERTY()
diff --git a/Engine/Source/Runtime/Engine/Private/Blueprint.cpp b/Engine/Source/Runtime/Engine/Private/Blueprint.cpp
--- a/Engine/Source/Runtime/Engine/Private/Blueprint.cpp
+++ b/Engine/Source/Runtime/Engine/Private/Blueprint.cpp
@@ -1729,6 +1729,14 @@ FName UBlueprint::FindBlueprintPropertyNameFromGuid(const FGuid& PropertyGuid) c
 			return BPVarDesc.VarName;
 		}
 	}
+
+	for (const FBPVariableDescription& BPVarDesc : GeneratedVariables)
+	{
+		if (BPVarDesc.VarGuid == PropertyGuid)
+		{
+			return BPVarDesc.VarName;
+		}
+	}
 #endif
 
 	return NAME_None;
@@ -1744,6 +1752,14 @@ FGuid UBlueprint::FindBlueprintPropertyGuidFromName(const FName PropertyName) co
 			return BPVarDesc.VarGuid;
 		}
 	}
+
+	for (const FBPVariableDescription& BPVarDesc : GeneratedVariables)
+	{
+		if (BPVarDesc.VarName == PropertyName)
+		{
+			return BPVarDesc.VarGuid;
+		}
+	}
 #endif
 
 	return FGuid();
