Added extra information to SetFunction assert to track down the broken asset/function

#rnx

diff --git a/Engine/Source/Editor/KismetCompiler/Private/KismetCompilerMisc.cpp b/Engine/Source/Editor/KismetCompiler/Private/KismetCompilerMisc.cpp
--- a/Engine/Source/Editor/KismetCompiler/Private/KismetCompilerMisc.cpp
+++ b/Engine/Source/Editor/KismetCompiler/Private/KismetCompilerMisc.cpp
@@ -847,7 +847,7 @@ UEdGraphPin* FKismetCompilerUtilities::GenerateAssignmentNodes(class FKismetComp
 			if (!SetFunctionName.IsEmpty())
 			{
 				UFunction* SetFunction = ForClass->FindFunctionByName(*SetFunctionName);
-				check(SetFunction);
+				checkf(SetFunction, TEXT("Failed to find SetFunction '%s' on class '%s'!"), *SetFunctionName, *ForClass->GetPathName());
 
 				// Add a cast node so we can call the Setter function with a pin of the right class
 				UK2Node_DynamicCast* CastNode = CompilerContext.SpawnIntermediateNode<UK2Node_DynamicCast>(SpawnNode, SourceGraph);
