PR #12402: Correcting what looks like copy paste error resulting in using incorrect variables in FKismetCompilerUtilities::GenerateFieldNotificationSetNode

diff --git a/Engine/Source/Editor/KismetCompiler/Private/KismetCompilerMisc.cpp b/Engine/Source/Editor/KismetCompiler/Private/KismetCompilerMisc.cpp
--- a/Engine/Source/Editor/KismetCompiler/Private/KismetCompilerMisc.cpp
+++ b/Engine/Source/Editor/KismetCompiler/Private/KismetCompilerMisc.cpp
@@ -1106,8 +1106,8 @@ TTuple<UEdGraphPin*, UEdGraphPin*> FKismetCompilerUtilities::GenerateFieldNotifi
 	// Set Net args pin
 	{
 		K2Schema->TrySetDefaultValue(*CallFuncNode->FindPinChecked(FName("bHasLocalRepNotify")), bHasLocalRepNotify ? TEXT("True") : TEXT("False"));
-		K2Schema->TrySetDefaultValue(*CallFuncNode->FindPinChecked(FName("bShouldFlushDormancyOnSet")), bHasLocalRepNotify ? TEXT("True") : TEXT("False"));
-		K2Schema->TrySetDefaultValue(*CallFuncNode->FindPinChecked(FName("bIsNetProperty")), bHasLocalRepNotify ? TEXT("True") : TEXT("False"));
+		K2Schema->TrySetDefaultValue(*CallFuncNode->FindPinChecked(FName("bShouldFlushDormancyOnSet")), bShouldFlushDormancyOnSet ? TEXT("True") : TEXT("False"));
+		K2Schema->TrySetDefaultValue(*CallFuncNode->FindPinChecked(FName("bIsNetProperty")), bIsNetProperty ? TEXT("True") : TEXT("False"));
 	}
 
 	TTuple<UEdGraphPin*, UEdGraphPin*> Result = {CallFuncNode->GetExecPin(), CallFuncNode->GetThenPin()};
