Disable ticking if the blueprint Tick function is empty.

diff --git a/Engine/Source/Editor/KismetCompiler/Private/KismetCompiler.cpp b/Engine/Source/Editor/KismetCompiler/Private/KismetCompiler.cpp
--- a/Engine/Source/Editor/KismetCompiler/Private/KismetCompiler.cpp
+++ b/Engine/Source/Editor/KismetCompiler/Private/KismetCompiler.cpp
@@ -5566,6 +5566,21 @@ void FKismetCompilerContext::SetCanEverTick() const
 				MessageLog.Warning(*ConlictingMetadataWarning);
 			}
 
+			// If the tick node is disconnected we can avoid ticking needlessly
+			bool bIsEmptyBlueprintTick = false;
+			for (const FKismetFunctionContext& FunctionContext : FunctionList)
+			{
+				if (FunctionContext.Function->GetName() == ReceiveTickName)
+				{
+					if (FunctionContext.SourceEventFromStubGraph)
+					{
+						UEdGraphPin* ExecThenPin = FunctionContext.SourceEventFromStubGraph->FindPin(UEdGraphSchema_K2::PN_Then);
+						bIsEmptyBlueprintTick = ExecThenPin && ExecThenPin->LinkedTo.IsEmpty();
+					}
+					break;
+				}
+			}
+
 			if (bHasCannotTickMetadata)
 			{
 				// This could only happen if someone adds bad metadata to AActor or UActorComponent directly
@@ -5578,7 +5593,7 @@ void FKismetCompilerContext::SetCanEverTick() const
 				).ToString();
 				MessageLog.Warning(*NativeClassSaidNo, FindLocalEntryPoint(ReceiveTickEvent));
 			}
-			else
+			else if (!bIsEmptyBlueprintTick)
 			{
 				if (bAllowTickingByDefault || bHasUniversalParent || bHasCanTickMetadata)
 				{
