Fix issue with function not being created yet for a CreateWidget expanded node because the widget class is in the queue for the full recompile
The function reference is stored as the authoritative class and the function name, so generating it from the skeleton class or generated class actually gives the exact same result

#tests Tested with the project in jira and loading it without crashing
#rnx

diff --git a/Engine/Source/Editor/KismetCompiler/Private/KismetCompilerMisc.cpp b/Engine/Source/Editor/KismetCompiler/Private/KismetCompilerMisc.cpp
--- a/Engine/Source/Editor/KismetCompiler/Private/KismetCompilerMisc.cpp
+++ b/Engine/Source/Editor/KismetCompiler/Private/KismetCompilerMisc.cpp
@@ -847,6 +847,13 @@ UEdGraphPin* FKismetCompilerUtilities::GenerateAssignmentNodes(class FKismetComp
 			if (!SetFunctionName.IsEmpty())
 			{
 				UFunction* SetFunction = ForClass->FindFunctionByName(*SetFunctionName);
+				if (SetFunction == nullptr)
+				{
+					if (UBlueprint* Generator = Cast<UBlueprint>(ForClass->ClassGeneratedBy))
+					{
+						SetFunction = Generator->SkeletonGeneratedClass->FindFunctionByName(*SetFunctionName);
+					}
+				}
 				checkf(SetFunction, TEXT("Failed to find SetFunction '%s' on class '%s'!"), *SetFunctionName, *ForClass->GetPathName());
 
 				// Add a cast node so we can call the Setter function with a pin of the right class
