[BP] Harvest metadata from user-defined function local variable descriptors at compile time to support details customizations that may expose it for editing.

diff --git a/Engine/Source/Editor/KismetCompiler/Private/KismetCompiler.cpp b/Engine/Source/Editor/KismetCompiler/Private/KismetCompiler.cpp
--- a/Engine/Source/Editor/KismetCompiler/Private/KismetCompiler.cpp
+++ b/Engine/Source/Editor/KismetCompiler/Private/KismetCompiler.cpp
@@ -1752,6 +1752,12 @@ FProperty* FKismetCompilerContext::CreateUserDefinedLocalVariableForFunction(con
 		NewProperty->SetMetaData(TEXT("Category"), *Variable.Category.ToString());
 		NewProperty->RepNotifyFunc = Variable.RepNotifyFunc;
 		NewProperty->SetPropertyFlags((EPropertyFlags)Variable.PropertyFlags);
+
+		// Set metadata on property (some customizations may expose it)
+		for (const FBPVariableMetaDataEntry& Entry : Variable.MetaDataArray)
+		{
+			NewProperty->SetMetaData(Entry.DataKey, *Entry.DataValue);
+		}
 	}
 	
 	return NewProperty;
