Support object literals for interface parameters

[RN] Blueprint function calls with literal values supplied for interface parameters now compile correctly

diff --git a/Engine/Source/Editor/KismetCompiler/Private/KismetCompilerVMBackend.cpp b/Engine/Source/Editor/KismetCompiler/Private/KismetCompilerVMBackend.cpp
--- a/Engine/Source/Editor/KismetCompiler/Private/KismetCompilerVMBackend.cpp
+++ b/Engine/Source/Editor/KismetCompiler/Private/KismetCompilerVMBackend.cpp
@@ -1109,7 +1109,13 @@ public:
 				}
 				else
 				{
-					ensureMsgf(false, TEXT("It is not possible to express this interface property as a literal value! (%s)"), *CoerceProperty->GetFullName());
+					// emit cast
+					Writer << EX_ObjToInterfaceCast;
+					UClass* InterfaceClass = CastChecked<UClass>(Term->Type.PinSubCategoryObject.Get());
+					Writer << InterfaceClass;
+					// from the object:
+					Writer << EX_ObjectConst;
+					Writer << Term->ObjectLiteral;
 				}
 			}
 			else if (!CoerceProperty && Term->Type.PinCategory.IsNone() && (Term->Type.PinSubCategory == UEdGraphSchema_K2::PN_Self))
