Emit correct bytecode for script interface self params

[RN] Interface parameters in blueprints that are assigned to 'self' on the function call node now behave correctly

diff --git a/Engine/Source/Editor/KismetCompiler/Private/KismetCompilerVMBackend.cpp b/Engine/Source/Editor/KismetCompiler/Private/KismetCompilerVMBackend.cpp
--- a/Engine/Source/Editor/KismetCompiler/Private/KismetCompilerVMBackend.cpp
+++ b/Engine/Source/Editor/KismetCompiler/Private/KismetCompilerVMBackend.cpp
@@ -1101,6 +1101,9 @@ public:
 			{
 				if (Term->Type.PinSubCategory == UEdGraphSchema_K2::PN_Self)
 				{
+					Writer << EX_ObjToInterfaceCast;
+					UClass* InterfaceClass = CastChecked<UClass>(Term->Type.PinSubCategoryObject.Get());
+					Writer << InterfaceClass;
 					Writer << EX_Self;
 				}
 				else if (Term->ObjectLiteral == nullptr)
