Fixed incorrect format patterns used with FText::Format

#rnx

diff --git a/Engine/Plugins/Experimental/DynamicMaterial/Source/DynamicMaterialEditor/Private/Utils/DMPrivate.cpp b/Engine/Plugins/Experimental/DynamicMaterial/Source/DynamicMaterialEditor/Private/Utils/DMPrivate.cpp
--- a/Engine/Plugins/Experimental/DynamicMaterial/Source/DynamicMaterialEditor/Private/Utils/DMPrivate.cpp
+++ b/Engine/Plugins/Experimental/DynamicMaterial/Source/DynamicMaterialEditor/Private/Utils/DMPrivate.cpp
@@ -97,7 +97,7 @@ namespace UE::DynamicMaterialEditor::Private
 	}
 
 	static bool bAllowUIFeedback = false;
-	static FText LogErrorObjectFormat = LOCTEXT("LogErrorObjectFormat", "%s (Source: %s)");
+	static FText LogErrorObjectFormat = LOCTEXT("LogErrorObjectFormat", "{0} (Source: {1})");
 
 	void LogError(const FString& InMessage, bool bInToast, const UObject* InSource)
 	{
diff --git a/Engine/Source/Editor/KismetCompiler/Private/KismetCompiler.cpp b/Engine/Source/Editor/KismetCompiler/Private/KismetCompiler.cpp
--- a/Engine/Source/Editor/KismetCompiler/Private/KismetCompiler.cpp
+++ b/Engine/Source/Editor/KismetCompiler/Private/KismetCompiler.cpp
@@ -5564,7 +5564,7 @@ void FKismetCompilerContext::SetCanEverTick() const
 			{
 				// User error: The C++ class has conflicting metadata
 				const FString ConlictingMetadataWarning = FText::Format(
-					LOCTEXT("HasBothCanAndCannotMetadataFmt", "Native class %s has both '{0}' and '{1}' metadata specified, they are mutually exclusive and '{1}' will win."),
+					LOCTEXT("HasBothCanAndCannotMetadataFmt", "Native class {0} has both '{1}' and '{2}' metadata specified, they are mutually exclusive and '{2}' will win."),
 					FText::FromString(FirstNativeClass->GetPathName()),
 					FText::FromName(FBlueprintMetadata::MD_ChildCanTick),
 					FText::FromName(FBlueprintMetadata::MD_ChildCannotTick)
