Fixed DefaultToInstanced subobjects breaking on blueprint duplication.
This was caused by SkeletonOnly blueprint compilation causing a CDO to be created, while it's possible that the non-skeleton blueprint generated class doesn't have initialized class flags. This is the case when duplicating a blueprint since a CDO is created before the class was ever fully compiled. This caused instanced subobject restoration to be skipped in the CDO's PostConstructInit->InstanceSubobjects, leading to invalid instanced default subobjects on the duplicate blueprint.

diff --git a/Engine/Source/Editor/Kismet/Private/BlueprintCompilationManager.cpp b/Engine/Source/Editor/Kismet/Private/BlueprintCompilationManager.cpp
--- a/Engine/Source/Editor/Kismet/Private/BlueprintCompilationManager.cpp
+++ b/Engine/Source/Editor/Kismet/Private/BlueprintCompilationManager.cpp
@@ -4,6 +4,7 @@
 
 #include "Async/Async.h"
 #include "Kismet2/BlueprintEditorUtils.h"
+#include "BlueprintClassFlagUtils.h"
 #include "BlueprintCompilerExtension.h"
 #include "BlueprintEditorSettings.h"
 #include "Blueprint/BlueprintSupport.h"
@@ -1511,6 +1512,17 @@ void FBlueprintCompilationManagerImpl::FlushCompilationQueueImpl(bool bSuppressB
 					BPGC->bLayoutChanging = false;
 					BPGC->Bind();
 					BPGC->StaticLink(true);
+
+					// When compiling skeleton only, it's possible that a non-skeleton BPGC
+					// has not been compiled yet and has uninitialized class flags. Some class
+					// flags must be set in order for the CDO construction to work correctly,
+					// like for instancing subobjects. Make sure that BPGC has at least its
+					// inherited and property based class flags.
+					FBlueprintClassFlagUtils::PropagateParentClassFlags_CompileClassLayout(BPGC);
+					FBlueprintClassFlagUtils::PropagateParentClassFlags_FinishCompilingClass(BPGC);
+					FBlueprintClassFlagUtils::AppendPropertyBasedClassFlags(BPGC);
+
+					// Force recreate the CDO
 					BPGC->SetDefaultObject(nullptr);
 					BPGC->GetDefaultObject(true);
 				}
diff --git a/Engine/Source/Editor/KismetCompiler/Internal/BlueprintClassFlagUtils.h b/Engine/Source/Editor/KismetCompiler/Internal/BlueprintClassFlagUtils.h
new file mode 100644
--- /dev/null
+++ b/Engine/Source/Editor/KismetCompiler/Internal/BlueprintClassFlagUtils.h
@@ -0,0 +1,25 @@
+// Copyright Epic Games, Inc. All Rights Reserved.
+
+#pragma once
+
+#include "CoreMinimal.h"
+
+#define UE_API KISMETCOMPILER_API
+
+class UClass;
+
+/** Utility functions for ClassFlags that is reused between normal and skeleton-only compilation. */
+class FBlueprintClassFlagUtils
+{
+public:
+	/** Given a class, let it inherit class flags from its parent class. Only considers inheritable flags. */
+	static UE_API void PropagateParentClassFlags_CompileClassLayout(UClass* Class);
+
+	/** Given a class, let it inherit class flags from its parent class. Only considers inheritable flags. */
+	static UE_API void PropagateParentClassFlags_FinishCompilingClass(UClass* Class);
+
+	/** Given a class, set any class flags that are based on its properties. Does not consider parent class properties. */
+	static UE_API void AppendPropertyBasedClassFlags(UClass* Class);
+};
+
+#undef UE_API
diff --git a/Engine/Source/Editor/KismetCompiler/Private/BlueprintClassFlagUtils.cpp b/Engine/Source/Editor/KismetCompiler/Private/BlueprintClassFlagUtils.cpp
new file mode 100644
--- /dev/null
+++ b/Engine/Source/Editor/KismetCompiler/Private/BlueprintClassFlagUtils.cpp
@@ -0,0 +1,59 @@
+// Copyright Epic Games, Inc. All Rights Reserved.
+
+#include "BlueprintClassFlagUtils.h"
+#include "UObject/Class.h"
+#include "UObject/ObjectMacros.h"
+#include "UObject/UnrealType.h"
+
+void FBlueprintClassFlagUtils::PropagateParentClassFlags_CompileClassLayout(UClass* Class)
+{
+	UClass* ParentClass = Class->GetSuperClass();
+	Class->ClassFlags |= (ParentClass->ClassFlags & CLASS_Inherit);
+	Class->ClassCastFlags |= ParentClass->ClassCastFlags;
+}
+
+void FBlueprintClassFlagUtils::PropagateParentClassFlags_FinishCompilingClass(UClass* Class)
+{
+	UClass* ParentClass = Class->GetSuperClass();
+	Class->ClassFlags &= ~CLASS_RecompilerClear;
+	Class->ClassFlags |= (ParentClass->ClassFlags & CLASS_ScriptInherit);//@TODO: ChangeParentClass had this, but I don't think I want it: | UClass::StaticClassFlags;  // will end up with CLASS_Intrinsic
+	Class->ClassCastFlags |= ParentClass->ClassCastFlags;
+	Class->ClassConfigName = ParentClass->ClassConfigName;
+	Class->ClassWithin = ParentClass->ClassWithin;
+}
+
+void FBlueprintClassFlagUtils::AppendPropertyBasedClassFlags(UClass* Class)
+{
+	for (TFieldIterator<FProperty> It(Class, EFieldIteratorFlags::ExcludeSuper); It; ++It)
+	{
+		FProperty* Property = *It;
+
+		// If any property is instanced, then the class needs to also have CLASS_HasInstancedReference flag
+		if (Property->ContainsInstancedObjectProperty())
+		{
+			Class->ClassFlags |= CLASS_HasInstancedReference;
+		}
+
+		// Look for OnRep 
+		if (Property->HasAnyPropertyFlags(CPF_Net))
+		{
+			// Verify rep notifies are valid, if not, clear them
+			if (Property->HasAnyPropertyFlags(CPF_RepNotify))
+			{
+				UFunction* OnRepFunc = Class->FindFunctionByName(Property->RepNotifyFunc);
+				if (OnRepFunc != NULL && OnRepFunc->NumParms == 0 && OnRepFunc->GetReturnProperty() == NULL)
+				{
+					// This function is good so just continue
+					continue;
+				}
+				// Invalid function for RepNotify! clear the flag
+				Property->RepNotifyFunc = NAME_None;
+			}
+		}
+		if (Property->HasAnyPropertyFlags(CPF_Config))
+		{
+			// If we have properties that are set from the config, then the class needs to also have CLASS_Config flags
+			Class->ClassFlags |= CLASS_Config;
+		}
+	}
+}
\ No newline at end of file
diff --git a/Engine/Source/Editor/KismetCompiler/Private/KismetCompiler.cpp b/Engine/Source/Editor/KismetCompiler/Private/KismetCompiler.cpp
--- a/Engine/Source/Editor/KismetCompiler/Private/KismetCompiler.cpp
+++ b/Engine/Source/Editor/KismetCompiler/Private/KismetCompiler.cpp
@@ -6,6 +6,7 @@
 
 
 #include "KismetCompiler.h"
+#include "BlueprintClassFlagUtils.h"
 #include "Algo/Find.h"
 #include "Engine/BlueprintGeneratedClass.h"
 #include "Misc/CoreMisc.h"
@@ -3132,11 +3133,7 @@ void FKismetCompilerContext::FinishCompilingClass(UClass* Class)
 	{
 		// Propagate the new parent's inheritable class flags
 		Class->ReferenceSchema.Reset();
-		Class->ClassFlags &= ~CLASS_RecompilerClear;
-		Class->ClassFlags |= (ParentClass->ClassFlags & CLASS_ScriptInherit);//@TODO: ChangeParentClass had this, but I don't think I want it: | UClass::StaticClassFlags;  // will end up with CLASS_Intrinsic
-		Class->ClassCastFlags |= ParentClass->ClassCastFlags;
-		Class->ClassConfigName = ParentClass->ClassConfigName;
-		Class->ClassWithin = ParentClass->ClassWithin;
+		FBlueprintClassFlagUtils::PropagateParentClassFlags_FinishCompilingClass(Class);
 
 		// If the Blueprint was marked as deprecated, then flag the class as deprecated.
 		if(Blueprint->bDeprecate)
@@ -3213,39 +3210,7 @@ void FKismetCompilerContext::FinishCompilingClass(UClass* Class)
 	Class->ClassFlags |= (CLASS_CompiledFromBlueprint);
 	Class->ClassFlags &= ~CLASS_ReplicationDataIsSetUp;
 
-	// This function mostly mirrors PostParsingClassSetup, opportunity to refactor:
-	for( TFieldIterator<FProperty> It(Class, EFieldIteratorFlags::ExcludeSuper); It; ++It)
-	{
-		FProperty *Property = *It;
-		
-		// If any property is instanced, then the class needs to also have CLASS_HasInstancedReference flag
-		if (Property->ContainsInstancedObjectProperty())
-		{
-			Class->ClassFlags |= CLASS_HasInstancedReference;
-		}
-		
-		// Look for OnRep 
-		if (Property->HasAnyPropertyFlags(CPF_Net))
-		{
-			// Verify rep notifies are valid, if not, clear them
-			if (Property->HasAnyPropertyFlags(CPF_RepNotify))
-			{
-				UFunction * OnRepFunc = Class->FindFunctionByName(Property->RepNotifyFunc);
-				if( OnRepFunc != NULL && OnRepFunc->NumParms == 0 && OnRepFunc->GetReturnProperty() == NULL )
-				{
-					// This function is good so just continue
-					continue;
-				}
-				// Invalid function for RepNotify! clear the flag
-				Property->RepNotifyFunc = NAME_None;
-			}
-		}
-		if( Property->HasAnyPropertyFlags( CPF_Config ))
-		{
-			// If we have properties that are set from the config, then the class needs to also have CLASS_Config flags
-			Class->ClassFlags |= CLASS_Config;
-		}
-	}
+	FBlueprintClassFlagUtils::AppendPropertyBasedClassFlags(Class);
 
 	// Verify class metadata as needed
 	if (FBlueprintEditorUtils::IsInterfaceBlueprint(Blueprint))
@@ -4886,9 +4851,7 @@ void FKismetCompilerContext::CompileClassLayout(EInternalCompilerFlags InternalF
 	NewClass->ClassGeneratedBy = Blueprint;
 	
 	// Set class metadata as needed
-	UClass* ParentClass = NewClass->GetSuperClass();
-	NewClass->ClassFlags |= (ParentClass->ClassFlags & CLASS_Inherit);
-	NewClass->ClassCastFlags |= ParentClass->ClassCastFlags;
+	FBlueprintClassFlagUtils::PropagateParentClassFlags_CompileClassLayout(NewClass);
 	
 	if (FBlueprintEditorUtils::IsInterfaceBlueprint(Blueprint))
 	{
