Add a temporary ensure in an effort to track down a bug in the blueprint compiler related to assignment generated for object construction

#rnx

diff --git a/Engine/Source/Editor/KismetCompiler/Private/KismetCompilerMisc.cpp b/Engine/Source/Editor/KismetCompiler/Private/KismetCompilerMisc.cpp
--- a/Engine/Source/Editor/KismetCompiler/Private/KismetCompilerMisc.cpp
+++ b/Engine/Source/Editor/KismetCompiler/Private/KismetCompilerMisc.cpp
@@ -765,10 +765,18 @@ bool FKismetCompilerUtilities::ValidateSelfCompatibility(const UEdGraphPin* Pin,
 
 UEdGraphPin* FKismetCompilerUtilities::GenerateAssignmentNodes(class FKismetCompilerContext& CompilerContext, UEdGraph* SourceGraph, UK2Node* CallBeginSpawnNode, UEdGraphNode* SpawnNode, UEdGraphPin* CallBeginResult, const UClass* ForClass, const UEdGraphPin* CallBeginClassInput)
 {
+	static const FName ProblemBPClassName(TEXT("QuestUpdatesLog_C"));
 	static const FName ObjectParamName(TEXT("Object"));
 	static const FName ValueParamName(TEXT("Value"));
 	static const FName PropertyNameParamName(TEXT("PropertyName"));
 
+	bool bDoTempValidate = false;
+	int32 NumAssignments = 0;
+	if(CompilerContext.NewClass && CompilerContext.NewClass->GetFName() == ProblemBPClassName)
+	{
+		bDoTempValidate = true;
+	}
+
 	const UEdGraphSchema_K2* Schema = CompilerContext.GetSchema();
 	UEdGraphPin* LastThen = CallBeginSpawnNode->GetThenPin();
 
@@ -911,6 +919,8 @@ UEdGraphPin* FKismetCompilerUtilities::GenerateAssignmentNodes(class FKismetComp
 			}
 			else if (UFunction* SetByNameFunction = Schema->FindSetVariableByNameFunction(OrgPin->PinType))
 			{
+				++NumAssignments;
+
 				UK2Node_CallFunction* SetVarNode = nullptr;
 				if (OrgPin->PinType.IsArray())
 				{
@@ -978,6 +988,17 @@ UEdGraphPin* FKismetCompilerUtilities::GenerateAssignmentNodes(class FKismetComp
 		}
 	}
 
+	if(bDoTempValidate)
+	{
+		ensureMsgf(
+			NumAssignments == 4, 
+				TEXT("Incorrect link set detected, %d %s %s %d"), 
+				NumAssignments, 
+				*ForClass->GetName(), 
+				ForClass->bLayoutChanging ? TEXT("layout changing") : TEXT("layout NOT changing"),
+				ForClass->PropertiesSize);
+	}
+
 	return LastThen;
 }
 
