Remove pointless null check found by new version of PVS Studio

#rnx

diff --git a/Engine/Source/Editor/KismetCompiler/Private/KismetCompilerMisc.cpp b/Engine/Source/Editor/KismetCompiler/Private/KismetCompilerMisc.cpp
--- a/Engine/Source/Editor/KismetCompiler/Private/KismetCompilerMisc.cpp
+++ b/Engine/Source/Editor/KismetCompiler/Private/KismetCompilerMisc.cpp
@@ -2320,7 +2320,7 @@ void FNodeHandlingFunctor::ResolveAndRegisterScopedTerm(FKismetFunctionContext&
 		{
 			Term->SetVarTypeLocal(true);
 		}
-		else if (BoundProperty->HasAnyPropertyFlags(CPF_BlueprintReadOnly) || (Context.IsConstFunction() && Context.NewClass && Context.NewClass->IsChildOf(SearchScope)))
+		else if (BoundProperty->HasAnyPropertyFlags(CPF_BlueprintReadOnly) || (Context.IsConstFunction() && Context.NewClass->IsChildOf(SearchScope)))
 		{
 			// Read-only variables and variables in const classes are both const
 			Term->bIsConst = true;
