Don't check local variables (if present) when computing DoSignaturesHaveConvertibleFloatTypes

[RN] Blueprint's delegate signature checking logic no longer accidently looks at local variables when matching function signatures

diff --git a/Engine/Source/Editor/KismetCompiler/Private/KismetCompilerMisc.cpp b/Engine/Source/Editor/KismetCompiler/Private/KismetCompilerMisc.cpp
--- a/Engine/Source/Editor/KismetCompiler/Private/KismetCompilerMisc.cpp
+++ b/Engine/Source/Editor/KismetCompiler/Private/KismetCompilerMisc.cpp
@@ -2255,10 +2255,20 @@ ConvertibleSignatureMatchResult FKismetCompilerUtilities::DoSignaturesHaveConver
 		TFieldIterator<FProperty> PropAIt(A);
 		TFieldIterator<FProperty> PropBIt(B);
 
-		while (PropAIt)
+		while (PropAIt && (PropAIt->PropertyFlags & CPF_Parm))
 		{
-			if (PropBIt)
+			if (PropBIt && (PropBIt->PropertyFlags & CPF_Parm))
 			{
+				// @todo: we should probably check these flags, but I'm unsure
+				// if it's strictly necessary to prevent binding e.g. a blueprint function
+				// that takes a struct by value from a delegate that takes a reference:
+				//const uint64 PropertyMash = PropAIt->PropertyFlags ^ PropBIt->PropertyFlags;
+				//const uint64 IgnoreFlags = UFunction::GetDefaultIgnoredSignatureCompatibilityFlags();
+				//if((PropertyMash & ~IgnoreFlags) != 0)
+				//{
+				//	return ConvertibleSignatureMatchResult::Different;
+				//}
+
 				if (!FStructUtils::ArePropertiesTheSame(*PropAIt, *PropBIt, false))
 				{
 					bool bHasConvertibleProperties =
@@ -2282,7 +2292,7 @@ ConvertibleSignatureMatchResult FKismetCompilerUtilities::DoSignaturesHaveConver
 		}
 
 		// If PropBIt still has parameters, then there was a mismatch
-		return PropBIt ? ConvertibleSignatureMatchResult::Different : ConvertibleSignatureMatchResult::HasConvertibleFloatParams;
+		return (PropBIt && PropBIt->PropertyFlags & CPF_Parm) ? ConvertibleSignatureMatchResult::Different : ConvertibleSignatureMatchResult::HasConvertibleFloatParams;
 	}
 
 	return ConvertibleSignatureMatchResult::ExactMatch;
