Edigrate 45550818
Harden FKismetCompilerContext::SetCanEverTick against invalid functions in the BP compiler's FunctionList
[RN] Fixed crash during bluepritn compilation due to missing function

diff --git a/Engine/Source/Editor/KismetCompiler/Private/KismetCompiler.cpp b/Engine/Source/Editor/KismetCompiler/Private/KismetCompiler.cpp
--- a/Engine/Source/Editor/KismetCompiler/Private/KismetCompiler.cpp
+++ b/Engine/Source/Editor/KismetCompiler/Private/KismetCompiler.cpp
@@ -5539,7 +5539,7 @@ void FKismetCompilerContext::SetCanEverTick() const
 			bool bIsEmptyBlueprintTick = false;
 			for (const FKismetFunctionContext& FunctionContext : FunctionList)
 			{
-				if (FunctionContext.Function->GetName() == ReceiveTickName)
+				if (FunctionContext.IsValid() && FunctionContext.Function->GetName() == ReceiveTickName)
 				{
 					if (FunctionContext.SourceEventFromStubGraph)
 					{
