Fixed a case of UFunction cook differences with compact set enabled.

#tests local asset cooks

diff --git a/Engine/Source/Editor/KismetCompiler/Private/KismetCompilerMisc.cpp b/Engine/Source/Editor/KismetCompiler/Private/KismetCompilerMisc.cpp
--- a/Engine/Source/Editor/KismetCompiler/Private/KismetCompilerMisc.cpp
+++ b/Engine/Source/Editor/KismetCompiler/Private/KismetCompilerMisc.cpp
@@ -7,6 +7,7 @@
 #include "KismetCompilerMisc.h"
 
 #include "BlueprintCompilationManager.h"
+#include "Containers/SparseSet.h"
 #include "Misc/CoreMisc.h"
 #include "UObject/MetaData.h"
 #include "UObject/UnrealType.h"
@@ -2711,7 +2712,7 @@ void FKismetFunctionContext::FinalSortLinearExecList()
 		return !(CurrentNode && CurStatementList && CurStatementList->Num());
 	});
 
-	TSet<UEdGraphNode*> UnsortedExecutionSet(LinearExecutionList);
+	TSparseSet<UEdGraphNode*> UnsortedExecutionSet(LinearExecutionList);
 	LinearExecutionList.Empty();
 	TArray<UEdGraphNode*> SortedLinearExecutionList;
 
@@ -2719,7 +2720,7 @@ void FKismetFunctionContext::FinalSortLinearExecList()
 	SortedLinearExecutionList.Push(EntryPoint);
 	UnsortedExecutionSet.Remove(EntryPoint);
 
-	TSet<UEdGraphNode*> NodesToStartNextChain;
+	TSparseSet<UEdGraphNode*> NodesToStartNextChain;
 
 	while (UnsortedExecutionSet.Num())
 	{
