Allow for clients to force edgraph transient flag

#rnx
#Synced-CL 48375597

diff --git a/Engine/Source/Editor/KismetCompiler/Private/KismetCompiler.cpp b/Engine/Source/Editor/KismetCompiler/Private/KismetCompiler.cpp
--- a/Engine/Source/Editor/KismetCompiler/Private/KismetCompiler.cpp
+++ b/Engine/Source/Editor/KismetCompiler/Private/KismetCompiler.cpp
@@ -64,6 +64,11 @@
 #include "UObject/UnrealTypePrivate.h"
 #include "KismetCompilerModule.h"
 
+namespace UE::Private
+{
+	KISMETCOMPILER_API bool (*GCanSaveFunctionsForPackage)(const UPackage*) = nullptr;
+}
+
 static bool bDebugPropertyPropagation = false;
 
 FSimpleMulticastDelegate FKismetCompilerContext::OnPreCompile;
@@ -2292,6 +2297,11 @@ void FKismetCompilerContext::PrecompileFunction(FKismetFunctionContext& Context,
 
 		Context.Function = NewObject<UFunction>(NewClass, NewFunctionName, RF_Public);
 
+		if( UE::Private::GCanSaveFunctionsForPackage && 
+			!(*UE::Private::GCanSaveFunctionsForPackage)(Blueprint->GetPackage()))
+		{
+			Context.Function->SetFlags(RF_Transient);
+		}
 #if USE_TRANSIENT_SKELETON
 		// Propagate down transient settings from the class
 		if (NewClass->HasAnyFlags(RF_Transient))
